Added production metrics and dashboards for Engine Samples

diff --git a/Engine/Build/Horde/Dashboard/analytics.dashboard.json b/Engine/Build/Horde/Dashboard/analytics.dashboard.json
new file mode 100644
--- /dev/null
+++ b/Engine/Build/Horde/Dashboard/analytics.dashboard.json
@@ -0,0 +1,402 @@
+{
+    "analytics": [
+        {
+            "id": "core-loading-user",
+            "name": "Core Loading User",
+			"variables": [
+				{
+					"name": "Stream",
+					"group": "Build_BranchName",
+				},
+				{
+					"name": "LoadingName",
+					"group": "LoadingName",
+					"defaults": [
+						"TotalEditorStartup",
+						"LoadMap",
+						"TimeToPIE"
+					]
+				},
+				{
+					"name": "MapName",
+					"group": "MapName",
+					"defaults": [
+						"AncientWorld",
+						"Big_City_LVL",
+						"L_DefaultEditorOverview"
+					]
+				},
+				{
+					"name": "AssetType",
+					"group": "AssetType"
+				},
+				{
+					"name": "TargetPlatform",
+					"group": "Cook_Profile_TargetPlatforms"
+				},
+				{
+					"name": "Horde JobName",
+					"group": "Horde_JobName"
+				},
+				{
+					"name": "Horde StepName",
+					"group": "Horde_StepName"
+				},
+				{
+					"name": "StreamId",
+					"group": "StreamId"
+				},
+				{
+					"name": "JobName",
+					"group": "JobName"
+				},
+				{
+					"name": "StepName",
+					"group": "StepName"
+				}
+			],
+            "categories": [
+                {
+                    "name": "Editor Iteration",
+					"charts": [
+						{
+							"name": "User Load Time (75th Percentile)",
+							"metrics": [
+								{
+									"id": "core-loading-user-timetopie-percentile",
+									"alias": "Time To PIE",
+									"threshold": 360
+								},
+								{
+									"id": "core-loading-user-editorstartuptime-percentile",
+									"alias": "Editor Startup",
+									"threshold": 120
+								},
+								{
+									"id": "core-loading-user-loadmaptime-percentile",
+									"alias": "Load Map",
+									"threshold": 240
+								},
+								{
+									"id": "core-loading-user-piestartuptime-percentile",
+									"alias": "PIE Startup",
+									"threshold": 360
+								}
+							],
+							"graph": "indicator",
+							"max": 720
+						},
+						{
+							"name": "User Load Time (Average)",
+							"metrics": [
+								{
+									"id": "core-loading-user-timetopie",
+									"alias": "Time To PIE",
+									"threshold": 360
+								},
+								{
+									"id": "core-loading-user-editorstartuptime",
+									"alias": "Editor Startup",
+									"threshold": 120
+								},
+								{
+									"id": "core-loading-user-loadmaptime",
+									"alias": "Load Map",
+									"threshold": 240
+								},
+								{
+									"id": "core-loading-user-piestartuptime",
+									"alias": "PIE Startup",
+									"threshold": 360
+								}
+							],
+							"graph": "indicator",
+							"max": 900
+						},
+						{
+							"name": "User Loading Time (KPIs)",
+							"metrics": [
+								{
+									"id": "core-loading-user-editorstartuptime-kpi",
+									"alias": "TotalEditorStartup"
+								},
+								{
+									"id": "core-loading-user-loadmaptime-kpi",
+									"alias": "LoadMap"
+								},
+								{
+									"id": "core-loading-user-piestartuptime-kpi",
+									"alias": "PIE.TotalStartupTime"
+								},
+								{
+									"id": "core-loading-user-timetopie-kpi",
+									"alias": "TimeToPIE"
+								}
+							]
+						},
+						{
+							"name": "User Loading Count (Live)",
+							"display": "Value",
+							"max" : 100,
+							"metrics": [
+								{
+									"id": "core-loading-user-loadingcount"
+								}
+							]
+						},
+						{
+							"name": "User Editor Startup Time (Live)",
+							"metrics": [
+								{
+									"id": "core-loading-user-editorstartuptime"
+								}
+							]
+						},
+						{
+							"name": "User Editor LoadMap Time (Live)",
+							"metrics": [
+								{
+									"id": "core-loading-user-loadmaptime"
+								}
+							]
+						},
+						{
+							"name": "User Editor LoadMap Count (Live)",
+							"display": "Value",
+							"metrics": [
+								{
+									"id": "core-loading-user-loadmap-count"
+								}
+							]
+						},
+						{
+							"name": "User PIE Startup Time (Live)",
+							"metrics": [
+								{
+									"id": "core-loading-user-piestartuptime"
+								}
+							]
+						},
+						{
+							"name": "User PIE LoadMap Time (Live)",
+							"metrics": [
+								{
+									"id": "core-loading-user-pieloadmaptime"
+								}
+							]
+						},
+						{
+							"name": "User Boot To PIE Time (Live)",
+							"metrics": [
+								{
+									"id": "core-loading-user-timetopie"
+								}
+							]
+						}
+					]
+                },
+                {
+                    "name": "Cooking",
+					"charts": [
+						{
+							"name": "Cook Time (KPIs)",
+							"metrics": [
+								{
+									"id": "core-cooking-nightly-cook-time-windows-kpi",
+									"alias": "Nightly Cook Windows",
+									"threshold": 7200
+								},
+								{
+									"id": "core-cooking-nightly-cook-time-ps5-kpi",
+									"alias": "Nightly Cook PS5",
+									"threshold": 7200
+								},
+								{
+									"id": "core-cooking-autotest-cook-time-kpi",
+									"alias": "AutoTest Cook",
+									"threshold": 4500
+								},
+								{
+									"id": "core-cooking-autotest-fastcook-time-kpi",
+									"alias": "AutoTest FastCook",
+									"threshold": 1800
+								}
+							],
+							"graph": "indicator",
+							"max": 10800
+						},
+						{
+							"name": "Cook Shaders Compiled (KPIs)",
+							"display": "Value",
+							"metrics": [
+								{
+									"id": "core-cooking-nightly-cook-shaders-compiled-windows-kpi",
+									"alias": "Nightly Cook Windows",
+									"threshold": 1000
+								},
+								{
+									"id": "core-cooking-nightly-cook-shaders-compiled-ps5-kpi",
+									"alias": "Nightly Cook PS5",
+									"threshold": 1000
+								},
+								{
+									"id": "core-cooking-autotest-cook-shaders-compiled-kpi",
+									"alias": "AutoTest Cook",
+									"threshold": 1000
+								},
+								{
+									"id": "core-cooking-autotest-fastcook-shaders-compiled-kpi",
+									"alias": "AutoTest FastCook",
+									"threshold": 1000
+								}
+							],
+							"graph": "indicator",
+							"max": 2000
+						},
+						{
+							"name": "Cook Time (KPIs)",
+							"metrics": [
+								{
+									"id": "core-cooking-nightly-cook-time-windows-kpi",
+									"alias": "Nightly Cook Windows"
+								},
+								{
+									"id": "core-cooking-nightly-cook-time-ps5-kpi",
+									"alias": "Nightly Cook PS5"
+								},
+								{
+									"id": "core-cooking-autotest-cook-time-kpi",
+									"alias": "AutoTest Cook"
+								},
+								{
+									"id": "core-cooking-autotest-fastcook-time-kpi",
+									"alias": "AutoTest FastCook"
+								}
+							]
+						},
+						{
+							"name": "Cook Shaders Compiled (KPIs)",
+							"display": "Value",
+							"metrics": [
+								{
+									"id": "core-cooking-nightly-cook-shaders-compiled-windows-kpi",
+									"alias": "Nightly Cook Windows"
+								},
+								{
+									"id": "core-cooking-nightly-cook-shaders-compiled-ps5-kpi",
+									"alias": "Nightly Cook PS5"
+								},
+								{
+									"id": "core-cooking-autotest-cook-shaders-compiled-kpi",
+									"alias": "AutoTest Cook"
+								},
+								{
+									"id": "core-cooking-autotest-fastcook-shaders-compiled-kpi",
+									"alias": "AutoTest FastCook"
+								}
+							]
+						},
+						{
+							"name": "Nightly Cook Time",
+							"metrics": [
+								{
+									"id": "core-cooking-nightly-cook-time"
+								}
+							]
+						},
+						{
+							"name": "Nightly Cook Shaders Compiled",
+							"display": "Value",
+							"metrics": [
+								{
+									"id": "core-cooking-nightly-shaders-compiled"
+								}
+							]
+						},
+						{
+							"name": "Incremental Cook Time",
+							"metrics": [
+								{
+									"id": "core-cooking-incremental-cook-time"
+								}
+							]
+						},
+						{
+							"name": "Incremental Cook Shaders Compiled",
+							"metrics": [
+								{
+									"id": "core-cooking-incremental-shaders-compiled"
+								}
+							]
+						},
+						{
+							"name": "Standard Cook Time",
+							"metrics": [
+								{
+									"id": "core-cooking-standard-cook-time"
+								}
+							]
+						},
+						{
+							"name": "Standard Cook Shaders Compiled",
+							"metrics": [
+								{
+									"id": "core-cooking-standard-shaders-compiled"
+								}
+							]
+						}
+					]
+                },
+				{
+					"name": "DDC",
+					"charts": [
+						{
+							"name": "User DDC Backend Efficiency",
+							"display": "Ratio",
+							"metrics": [
+								{
+									"id": "core-loading-user-totalddchitratio",
+									"alias": "Root"
+								},
+								{
+									"id": "core-loading-user-cloudlocalddchitratio",
+									"alias": "Cloud"
+								},
+								{
+									"id": "core-loading-user-zenlocalddchitratio",
+									"alias": "Local"
+								},
+								{
+									"id": "core-loading-user-sharedddchitratio",
+									"alias": "Shared"
+								}
+							]
+						},
+						{
+							"name": "User DDC Asset Efficiency",
+							"display": "Ratio",
+							"metrics": [
+								{
+									"id": "core-loading-user-ddc-asset-efficiency"
+								}
+							]
+						}
+					]
+				},
+				{
+					"name": "Build",
+					"charts": [
+						{
+							"name": "Horde Job Step Time",
+							"metrics": [
+								{
+									"id": "horde-jobstep-time"
+								}	
+							]
+						}
+					]
+				}
+            ]
+        }
+    ]
+}
\ No newline at end of file
diff --git a/Engine/Build/Horde/Telemetry/metrics.telemetry.json b/Engine/Build/Horde/Telemetry/metrics.telemetry.json
new file mode 100644
--- /dev/null
+++ b/Engine/Build/Horde/Telemetry/metrics.telemetry.json
@@ -0,0 +1,318 @@
+{
+	"metrics": [
+		{
+			"id": "core-loading-user-loadingtime",
+			"filter": "$[?(@.Payload.EventName == 'Core.Loading' && @.Payload.SchemaVersion>=2 && @.Payload.Config_IsBuildMachine==false && @.Payload.Config_IsRunningCommandlet==false && @.Payload.LoadingSeconds>5 && @.Payload.LoadingSeconds<3600)]",
+			"property": "$.Payload.LoadingSeconds",
+			"groupBy": "$.Payload.Build_BranchName,$.Payload.LoadingName",
+			"function": "Average",
+			"interval": "00:15:00"
+		},
+		{
+			"id": "core-loading-user-loadingcount",
+			"filter": "$[?(@.Payload.EventName == 'Core.Loading' && @.Payload.SchemaVersion>=2 &&@.Payload.Config_IsBuildMachine==false )]",
+			"groupBy": "$.Payload.Build_BranchName,$.Payload.LoadingName",
+			"function": "Count",
+			"interval": "00:15:00"
+		},
+		{
+			"id": "core-loading-user-editorstartuptime",
+			"filter": "$[?(@.Payload.EventName == 'Core.Loading' && @.Payload.SchemaVersion>=3 && @.Payload.LoadingName=='TotalEditorStartup' && @.Payload.Config_IsBuildMachine==false && @.Payload.Config_IsRunningCommandlet==false && @.Payload.LoadingSeconds>5 && @.Payload.LoadingSeconds<3600)]",
+			"property": "$.Payload.LoadingSeconds",
+			"groupBy": "$.Payload.Build_BranchName",
+			"function": "Average",
+			"interval": "00:15:00"
+		},
+		{
+			"id": "core-loading-user-loadmaptime",
+			"filter": "$[?(@.Payload.EventName == 'Core.Loading' && @.Payload.SchemaVersion>=2 && @.Payload.LoadingName=='LoadMap' && @.Payload.Config_IsBuildMachine==false && @.Payload.Config_IsRunningCommandlet==false && @.Payload.LoadingSeconds>5 && @.Payload.LoadingSeconds<3600)]",
+			"property": "$.Payload.LoadingSeconds",
+			"groupBy": "$.Payload.Build_BranchName,$.Payload.MapName",
+			"function": "Average",
+			"interval": "00:15:00"
+		},
+		{
+			"id": "core-loading-user-loadmap-count",
+			"filter": "$[?(@.Payload.EventName == 'Core.Loading' && @.Payload.SchemaVersion>=2 && @.Payload.LoadingName=='LoadMap' && @.Payload.Config_IsBuildMachine==false && @.Payload.Config_IsRunningCommandlet==false && @.Payload.LoadingSeconds>5 && @.Payload.LoadingSeconds<3600)]",
+			"groupBy": "$.Payload.Build_BranchName,$.Payload.MapName",
+			"function": "Count",
+			"interval": "00:15:00"
+		},
+		{
+			"id": "core-loading-user-pieloadmaptime",
+			"filter": "$[?(@.Payload.EventName == 'Core.Loading' && @.Payload.SchemaVersion>=2 && @.Payload.LoadingName=='PIE.LoadMapTime' && @.Payload.Config_IsBuildMachine==false && @.Payload.Config_IsRunningCommandlet==false && @.Payload.LoadingSeconds>5 && @.Payload.LoadingSeconds<3600)]",
+			"property": "$.Payload.LoadingSeconds",
+			"groupBy": "$.Payload.Build_BranchName,$.Payload.MapName",
+			"function": "Average",
+			"interval": "00:15:00"
+		},
+		{
+			"id": "core-loading-user-timetopie",
+			"filter": "$[?(@.Payload.EventName == 'Core.Loading' && @.Payload.SchemaVersion>=2 && @.Payload.LoadingName=='TimeToPIE' && @.Payload.Config_IsBuildMachine==false && @.Payload.Config_IsRunningCommandlet==false && @.Payload.LoadingSeconds>5 && @.Payload.LoadingSeconds<3600)]",
+			"property": "$.Payload.LoadingSeconds",
+			"groupBy": "$.Payload.Build_BranchName,$.Payload.MapName",
+			"function": "Average",
+			"interval": "00:15:00"
+		},
+		{
+			"id": "core-loading-user-piestartuptime",
+			"filter": "$[?(@.Payload.EventName == 'Core.Loading' && @.Payload.SchemaVersion>=2 && @.Payload.LoadingName=='PIE.TotalStartupTime' && @.Payload.Config_IsBuildMachine==false && @.Payload.Config_IsRunningCommandlet==false && @.Payload.LoadingSeconds>5 && @.Payload.LoadingSeconds<3600)]",
+			"property": "$.Payload.LoadingSeconds",
+			"groupBy": "$.Payload.Build_BranchName,$.Payload.MapName",
+			"function": "Average",
+			"interval": "00:15:00"
+		},
+		{
+			"id": "core-loading-user-editorstartuptime-kpi",
+			"filter": "$[?(@.Payload.EventName == 'Core.Loading' && @.Payload.SchemaVersion>=3 && @.Payload.LoadingName=='TotalEditorStartup' && @.Payload.Config_IsBuildMachine==false && @.Payload.Config_IsRunningCommandlet==false && @.Payload.LoadingSeconds>5 && @.Payload.LoadingSeconds<3600)]",
+			"property": "$.Payload.LoadingSeconds",
+			"groupBy": "$.Payload.Build_BranchName",
+			"function": "Average",
+			"interval": "01:00:00"
+		},
+		{
+			"id": "core-loading-user-loadmaptime-kpi",
+			"filter": "$[?(@.Payload.EventName == 'Core.Loading' && @.Payload.SchemaVersion>=2 && @.Payload.LoadingName=='LoadMap' && @.Payload.Config_IsBuildMachine==false && @.Payload.Config_IsRunningCommandlet==false && @.Payload.LoadingSeconds>5 && @.Payload.LoadingSeconds<3600)]",
+			"property": "$.Payload.LoadingSeconds",
+			"groupBy": "$.Payload.Build_BranchName",
+			"function": "Average",
+			"interval": "01:00:00"
+		},
+		{
+			"id": "core-loading-user-piestartuptime-kpi",
+			"filter": "$[?(@.Payload.EventName == 'Core.Loading' && @.Payload.SchemaVersion>=2 && @.Payload.LoadingName=='PIE.TotalStartupTime' && @.Payload.Config_IsBuildMachine==false && @.Payload.Config_IsRunningCommandlet==false && @.Payload.LoadingSeconds>5 && @.Payload.LoadingSeconds<3600)]",
+			"property": "$.Payload.LoadingSeconds",
+			"groupBy": "$.Payload.Build_BranchName",
+			"function": "Average",
+			"interval": "01:00:00"
+		},
+		{
+			"id": "core-loading-user-timetopie-kpi",
+			"filter": "$[?(@.Payload.EventName == 'Core.Loading' && @.Payload.SchemaVersion>=2 && @.Payload.LoadingName=='TimeToPIE' && @.Payload.Config_IsBuildMachine==false && @.Payload.Config_IsRunningCommandlet==false && @.Payload.LoadingSeconds>5 && @.Payload.LoadingSeconds<3600)]",
+			"property": "$.Payload.LoadingSeconds",
+			"groupBy": "$.Payload.Build_BranchName",
+			"function": "Average",
+			"interval": "01:00:00"
+		},
+		{
+			"id": "core-loading-user-cloudlocalddchitratio",
+			"filter": "$[?(@.Payload.EventName == 'Core.Loading' && @.Payload.SchemaVersion>=2 && @.Payload.Config_IsBuildMachine==false && @.Payload.DDC_Summary_CloudGetTotal>0 && @.Payload.LoadingSeconds>5 && @.Payload.LoadingSeconds<3600)]",
+			"property": "$.Payload.DDC_Summary_CloudGetHitPct",
+			"groupBy": "$.Payload.Build_BranchName",
+			"function": "Average",
+			"interval": "01:00:00"
+		},
+		{
+			"id": "core-loading-user-zenlocalddchitratio",
+			"filter": "$[?(@.Payload.EventName == 'Core.Loading' && @.Payload.SchemaVersion>=2 && @.Payload.Config_IsBuildMachine==false && @.Payload.Config_IsRunningCommandlet==false && @.Payload.DDC_Summary_ZenLocalGetTotal>0 && @.Payload.LoadingSeconds>5 && @.Payload.LoadingSeconds<3600)]",
+			"property": "$.Payload.DDC_Summary_ZenLocalGetHitPct",
+			"groupBy": "$.Payload.Build_BranchName",
+			"function": "Average",
+			"interval": "01:00:00"
+		},
+		{
+			"id": "core-loading-user-sharedddchitratio",
+			"filter": "$[?(@.Payload.EventName == 'Core.Loading' && @.Payload.SchemaVersion>=2 && @.Payload.Config_IsBuildMachine==false && @.Payload.Config_IsRunningCommandlet==false && @.Payload.DDC_Summary_SharedGetTotal>0 && @.Payload.LoadingSeconds>5 && @.Payload.LoadingSeconds<3600)]",
+			"property": "$.Payload.DDC_Summary_SharedGetHitPct",
+			"groupBy": "$.Payload.Build_BranchName",
+			"function": "Average",
+			"interval": "01:00:00"
+		},
+		{
+			"id": "core-loading-user-totalddchitratio",
+			"filter": "$[?(@.Payload.EventName == 'Core.Loading' && @.Payload.SchemaVersion>=2 && @.Payload.Config_IsBuildMachine==false && @.Payload.Config_IsRunningCommandlet==false && @.Payload.LoadingSeconds>5 && @.Payload.LoadingSeconds<3600)]",
+			"property": "$.Payload.DDC_Summary_TotalGetHitPct",
+			"groupBy": "$.Payload.Build_BranchName",
+			"function": "Average",
+			"interval": "01:00:00"
+		},
+		{
+			"id": "core-loading-user-ddc-asset-efficiency",
+			"filter": "$[?(@.Payload.EventName=='Core.DDC.Resource' && @.Payload.SchemaVersion>=3 && @.Payload.Config_IsBuildMachine==false && @.Payload.Config_IsEditor==true && @.Payload.AssetType!='Total')]",
+			"property": "$.Payload.Efficiency",
+			"groupBy": "$.Payload.Build_BranchName,$.Payload.AssetType",
+			"function": "Average",
+			"interval": "01:00:00"
+		},
+		{
+			"id": "core-loading-user-timetopie-percentile",
+			"filter": "$[?(@.Payload.EventName == 'Core.Loading' && @.Payload.SchemaVersion>=2 && @.Payload.LoadingName=='TimeToPIE' && @.Payload.Config_IsBuildMachine==false && @.Payload.Config_IsRunningCommandlet==false && @.Payload.LoadingSeconds>5 && @.Payload.LoadingSeconds<3600)]",
+			"groupBy": "$.Payload.Build_BranchName,$.Payload.MapName",
+			"property": "$.Payload.LoadingSeconds",
+			"function": "Percentile",
+			"percentile": 75,
+			"interval": "00:15:00"
+		},
+		{
+			"id": "core-loading-user-editorstartuptime-percentile",
+			"filter": "$[?(@.Payload.EventName == 'Core.Loading' && @.Payload.SchemaVersion>=2 && @.Payload.LoadingName=='TotalEditorStartup' && @.Payload.Config_IsBuildMachine==false && @.Payload.Config_IsRunningCommandlet==false && @.Payload.LoadingSeconds>5 && @.Payload.LoadingSeconds<3600 )]",
+			"groupBy": "$.Payload.Build_BranchName",
+			"property": "$.Payload.LoadingSeconds",
+			"function": "Percentile",
+			"percentile": 75,
+			"interval": "00:15:00"
+		},
+		{
+			"id": "core-loading-user-loadmaptime-percentile",
+			"filter": "$[?(@.Payload.EventName == 'Core.Loading' && @.Payload.SchemaVersion>=3 && @.Payload.LoadingName=='LoadMap' && @.Payload.Config_IsBuildMachine==false && @.Payload.LoadingSeconds>5 && @.Payload.LoadingSeconds<3600)]",
+			"groupBy": "$.Payload.Build_BranchName,$.Payload.MapName",
+			"property": "$.Payload.LoadingSeconds",
+			"function": "Percentile",
+			"percentile": 75,
+			"interval": "00:15:00"
+		},
+		{
+			"id": "core-loading-user-piestartuptime-percentile",
+			"filter": "$[?(@.Payload.EventName == 'Core.Loading' && @.Payload.SchemaVersion>=3 && @.Payload.LoadingName=='PIE.TotalStartupTime' && @.Payload.Config_IsBuildMachine==false && @.Payload.LoadingSeconds>5 && @.Payload.LoadingSeconds<3600)]",
+			"groupBy": "$.Payload.Build_BranchName,$.Payload.MapName",
+			"property": "$.Payload.LoadingSeconds",
+			"function": "Percentile",
+			"percentile": 75,
+			"interval": "00:15:00"
+		},
+		{
+			"id": "core-cooking-standard-cook-time",
+			"filter": "$[?(@.Payload.EventName == 'Core.Cooking' && @.Payload.SchemaVersion>=2 && search(@.Payload.Horde_TemplateName, \"Standard Build\") && search(@.Payload.Horde_StepName, \"Cook\"))]",
+			"groupBy": "$.Payload.Build_BranchName,$.Payload.Horde_StepName,$.Payload.Cook_Profile_TargetPlatforms",
+			"property": "$.Payload.Cook_Profile_CookWallTimeSec",
+			"function": "Average",
+			"interval": "01:00:00"
+		},
+		{
+			"id": "core-cooking-incremental-cook-time",
+			"filter": "$[?(@.Payload.EventName == 'Core.Cooking' && @.Payload.SchemaVersion>=3 && @.Payload.Horde_TemplateName=='Incremental CI')]",
+			"groupBy": "$.Payload.Build_BranchName,$.Payload.Horde_StepName",
+			"property": "$.Payload.Cook_Profile_CookWallTimeSec",
+			"function": "Average",
+			"interval": "00:30:00"
+		},
+		{
+			"id": "core-cooking-incremental-shaders-compiled",
+			"filter": "$[?(@.Payload.EventName == 'Core.Cooking' && @.Payload.SchemaVersion>=3 && @.Payload.Horde_TemplateName=='Incremental CI')]",
+			"groupBy": "$.Payload.Build_BranchName,$.Payload.Horde_StepName",
+			"property": "$.Payload.Shaders_ShadersCompiled",
+			"function": "Average",
+			"interval": "00:30:00"
+		},
+		{
+			"id": "core-cooking-nightly-cook-time",
+			"filter": "$[?(@.Payload.EventName == 'Core.Cooking' && @.Payload.SchemaVersion>=3 && search(@.Payload.Horde_TemplateName, \"Nightly Build\") && search(@.Payload.Horde_StepName, \"Cook\"))]",
+			"groupBy": "$.Payload.Build_BranchName,$.Payload.Horde_StepName,$.Payload.Cook_Profile_TargetPlatforms",
+			"property": "$.Payload.Cook_Profile_CookWallTimeSec",
+			"function": "Average",
+			"interval": "01:00:00"
+		},
+		{
+			"id": "core-cooking-nightly-cook-time-windows-kpi",
+			"filter": "$[?(@.Payload.EventName == 'Core.Cooking' && @.Payload.SchemaVersion>=3 && search(@.Payload.Horde_TemplateName, \"Nightly Build\") && @.Payload.Horde_StepName=='Cook Windows Client')]",
+			"groupBy": "$.Payload.Build_BranchName",
+			"property": "$.Payload.Cook_Profile_CookWallTimeSec",
+			"function": "Average",
+			"interval": "01:00:00"
+		},
+		{
+			"id": "core-cooking-nightly-cook-time-ps5-kpi",
+			"filter": "$[?(@.Payload.EventName == 'Core.Cooking' && @.Payload.SchemaVersion>=3 && search(@.Payload.Horde_TemplateName, \"Nightly Build\") && @.Payload.Horde_StepName=='Cook PS5 Client')]",
+			"groupBy": "$.Payload.Build_BranchName",
+			"property": "$.Payload.Cook_Profile_CookWallTimeSec",
+			"function": "Average",
+			"interval": "01:00:00"
+		},
+		{
+			"id": "core-cooking-autotest-cook-time-kpi",
+			"filter": "$[?(@.Payload.EventName == 'Core.Cooking' && @.Payload.SchemaVersion>=3 && @.Payload.Horde_StepName=='AutoTest Cook')]",
+			"groupBy": "$.Payload.Build_BranchName",
+			"property": "$.Payload.Cook_Profile_CookWallTimeSec",
+			"function": "Average",
+			"interval": "01:00:00"
+		},
+		{
+			"id": "core-cooking-autotest-fastcook-time-kpi",
+			"filter": "$[?(@.Payload.EventName == 'Core.Cooking' && @.Payload.SchemaVersion>=3 && @.Payload.Horde_StepName=='AutoTest FastCook')]",
+			"groupBy": "$.Payload.Build_BranchName",
+			"property": "$.Payload.Cook_Profile_CookWallTimeSec",
+			"function": "Average",
+			"interval": "01:00:00"
+		},
+		{
+			"id": "core-cooking-nightly-cook-shaders-compiled-windows-kpi",
+			"filter": "$[?(@.Payload.EventName == 'Core.Cooking' && @.Payload.SchemaVersion>=3 && search(@.Payload.Horde_TemplateName, \"Nightly Build\") && @.Payload.Horde_StepName=='Cook Windows Client')]",
+			"groupBy": "$.Payload.Build_BranchName",
+			"property": "$.Payload.Shaders_ShadersCompiled",
+			"function": "Average",
+			"interval": "01:00:00"
+		},
+		{
+			"id": "core-cooking-nightly-cook-shaders-compiled-ps5-kpi",
+			"filter": "$[?(@.Payload.EventName == 'Core.Cooking' && @.Payload.SchemaVersion>=3 && search(@.Payload.Horde_TemplateName, \"Nightly Build\") && @.Payload.Horde_StepName=='Cook PS5 Client')]",
+			"groupBy": "$.Payload.Build_BranchName",
+			"property": "$.Payload.Shaders_ShadersCompiled",
+			"function": "Average",
+			"interval": "01:00:00"
+		},
+		{
+			"id": "core-cooking-autotest-cook-shaders-compiled-kpi",
+			"filter": "$[?(@.Payload.EventName == 'Core.Cooking' && @.Payload.SchemaVersion>=3 && @.Payload.Horde_StepName=='AutoTest Cook')]",
+			"groupBy": "$.Payload.Build_BranchName",
+			"property": "$.Payload.Shaders_ShadersCompiled",
+			"function": "Average",
+			"interval": "01:00:00"
+		},
+		{
+			"id": "core-cooking-autotest-fastcook-shaders-compiled-kpi",
+			"filter": "$[?(@.Payload.EventName == 'Core.Cooking' && @.Payload.SchemaVersion>=3 && @.Payload.Horde_StepName=='AutoTest FastCook')]",
+			"groupBy": "$.Payload.Build_BranchName",
+			"property": "$.Payload.Shaders_ShadersCompiled",
+			"function": "Average",
+			"interval": "01:00:00"
+		},
+		{
+			"id": "core-cooking-standard-shaders-compiled",
+			"filter": "$[?(@.Payload.EventName == 'Core.Cooking' && @.Payload.SchemaVersion>=3 && search(@.Payload.Horde_TemplateName, \"Standard Build\"))]",
+			"groupBy": "$.Payload.Build_BranchName,$.Payload.Horde_StepName,$.Payload.Cook_Profile_TargetPlatforms",
+			"property": "$.Payload.Shaders_ShadersCompiled",
+			"function": "Average",
+			"interval": "01:00:00"
+		},
+		{
+			"id": "core-cooking-nightly-shaders-compiled",
+			"filter": "$[?(@.Payload.EventName == 'Core.Cooking' && @.Payload.SchemaVersion>=3 && search(@.Payload.Horde_TemplateName, \"Nightly Build\"))]",
+			"groupBy": "$.Payload.Build_BranchName,$.Payload.Horde_StepName,$.Payload.Cook_Profile_TargetPlatforms",
+			"property": "$.Payload.Shaders_ShadersCompiled",
+			"function": "Average",
+			"interval": "01:00:00"
+		},
+		{
+			"id": "core-cooking-autotest-cook-time",
+			"filter": "$[?(@.Payload.EventName == 'Core.Cooking' && @.Payload.SchemaVersion>=3 && @.Payload.Horde_StepName=='AutoTest Cook')]",
+			"groupBy": "$.Payload.Build_BranchName,$.Payload.Horde_StepName,$.Payload.Cook_Profile_TargetPlatforms",
+			"property": "$.Payload.Cook_Profile_CookWallTimeSec",
+			"function": "Average",
+			"interval": "00:30:00"
+		},
+		{
+			"id": "core-cooking-autotest-shaders-compiled",
+			"filter": "$[?(@.Payload.EventName == 'Core.Cooking' && @.Payload.SchemaVersion>=3 && @.Payload.Horde_StepName=='AutoTest Cook')]",
+			"groupBy": "$.Payload.Build_BranchName,$.Payload.Horde_StepName,$.Payload.Cook_Profile_TargetPlatforms",
+			"property": "$.Payload.Shaders_ShadersCompiled",
+			"function": "Average",
+			"interval": "00:30:00"
+		},
+		{
+			"id": "core-cooking-autotest-fastcook-time",
+			"filter": "$[?(@.Payload.EventName == 'Core.Cooking' && @.Payload.SchemaVersion>=3 && @.Payload.Horde_StepName=='AutoTest FastCook')]",
+			"groupBy": "$.Payload.Build_BranchName,$.Payload.Horde_StepName,$.Payload.Cook_Profile_TargetPlatforms",
+			"property": "$.Payload.Cook_Profile_CookWallTimeSec",
+			"function": "Average",
+			"interval": "01:00:00"
+		},
+		{
+			"id": "horde-jobstep-time",
+			"filter": "$[?(@.Payload.EventName == 'State.JobStepRef' && @.Payload.Duration>0 && @.Payload.State=='Completed')]",
+			"groupBy": "$.Payload.StreamId,$.Payload.JobName,$.Payload.StepName",
+			"property": "$.Payload.Duration",
+			"function": "Average",
+			"interval": "01:00:00"
+		}
+	]
+}
\ No newline at end of file
