Added Cook KPIs and Metrics for Lyra, CitySample and AncientGame

diff --git a/Engine/Build/Horde/Dashboard/analytics.dashboard.json b/Engine/Build/Horde/Dashboard/analytics.dashboard.json
--- a/Engine/Build/Horde/Dashboard/analytics.dashboard.json
+++ b/Engine/Build/Horde/Dashboard/analytics.dashboard.json
@@ -201,6 +201,36 @@
 						{
 							"name": "Cook Time (KPIs)",
 							"metrics": [
+								{
+									"id": "core-cooking-lyra-time-windows-kpi",
+									"alias": "Lyra Cook Windows",
+									"threshold": 7200
+								},
+								{
+									"id": "core-cooking-ancientgame-time-windows-kpi",
+									"alias": "AncientGame Cook Windows",
+									"threshold": 7200
+								},
+								{
+									"id": "core-cooking-citysample-time-windows-kpi",
+									"alias": "CitySample Cook Windows",
+									"threshold": 7200
+								},
+								{
+									"id": "core-cooking-lyra-time-ps5-kpi",
+									"alias": "Lyra Cook Windows",
+									"threshold": 7200
+								},
+								{
+									"id": "core-cooking-ancientgame-time-ps5-kpi",
+									"alias": "AncientGame Cook Windows",
+									"threshold": 7200
+								},
+								{
+									"id": "core-cooking-citysample-time-ps5-kpi",
+									"alias": "CitySample Cook Windows",
+									"threshold": 7200
+								},
 								{
 									"id": "core-cooking-nightly-cook-time-windows-kpi",
 									"alias": "Nightly Cook Windows",
@@ -297,35 +327,52 @@
 							]
 						},
 						{
-							"name": "Nightly Cook Time",
+							"name": "Incremental Cook Time",
 							"metrics": [
 								{
-									"id": "core-cooking-nightly-cook-time"
+									"id": "core-cooking-incremental-cook-time"
 								}
 							]
 						},
 						{
-							"name": "Nightly Cook Shaders Compiled",
+							"name": "Incremental Cook Shaders Compiled",
+							"metrics": [
+								{
+									"id": "core-cooking-incremental-shaders-compiled"
+								}
+							]
+						},
+						{
+							"name": "Package And Test Cook Time",
+							"metrics": [
+								{
+									"id": "core-cooking-package-and-test-cook-time"
+								}
+							]
+						},
+						{
+							"name": "Package And Test Shaders Compiled",
 							"display": "Value",
 							"metrics": [
 								{
-									"id": "core-cooking-nightly-shaders-compiled"
+									"id": "core-cooking-package-and-test-shaders-compiled"
 								}
 							]
 						},
 						{
-							"name": "Incremental Cook Time",
+							"name": "Nightly Cook Time",
 							"metrics": [
 								{
-									"id": "core-cooking-incremental-cook-time"
+									"id": "core-cooking-nightly-cook-time"
 								}
 							]
 						},
 						{
-							"name": "Incremental Cook Shaders Compiled",
+							"name": "Nightly Cook Shaders Compiled",
+							"display": "Value",
 							"metrics": [
 								{
-									"id": "core-cooking-incremental-shaders-compiled"
+									"id": "core-cooking-nightly-shaders-compiled"
 								}
 							]
 						},
diff --git a/Engine/Build/Horde/Telemetry/metrics.telemetry.json b/Engine/Build/Horde/Telemetry/metrics.telemetry.json
--- a/Engine/Build/Horde/Telemetry/metrics.telemetry.json
+++ b/Engine/Build/Horde/Telemetry/metrics.telemetry.json
@@ -2,7 +2,7 @@
 	"metrics": [
 		{
 			"id": "core-loading-user-loadingtime",
-			"filter": "$[?(@.Payload.EventName == 'Core.Loading' && @.Payload.SchemaVersion>=2 && @.Payload.Config_IsBuildMachine==false && @.Payload.Config_IsRunningCommandlet==false && @.Payload.LoadingSeconds>5 && @.Payload.LoadingSeconds<3600)]",
+			"filter": "$[?(@.Payload.EventName == 'Core.Loading' && @.Payload.Config_IsBuildMachine==false && @.Payload.Config_IsRunningCommandlet==false && @.Payload.LoadingSeconds>5 && @.Payload.LoadingSeconds<3600)]",
 			"property": "$.Payload.LoadingSeconds",
 			"groupBy": "$.Payload.Build_BranchName,$.Payload.LoadingName",
 			"function": "Average",
@@ -10,14 +10,14 @@
 		},
 		{
 			"id": "core-loading-user-loadingcount",
-			"filter": "$[?(@.Payload.EventName == 'Core.Loading' && @.Payload.SchemaVersion>=2 &&@.Payload.Config_IsBuildMachine==false )]",
+			"filter": "$[?(@.Payload.EventName == 'Core.Loading' &&@.Payload.Config_IsBuildMachine==false )]",
 			"groupBy": "$.Payload.Build_BranchName,$.Payload.LoadingName",
 			"function": "Count",
 			"interval": "00:15:00"
 		},
 		{
 			"id": "core-loading-user-editorstartuptime",
-			"filter": "$[?(@.Payload.EventName == 'Core.Loading' && @.Payload.SchemaVersion>=3 && @.Payload.LoadingName=='TotalEditorStartup' && @.Payload.Config_IsBuildMachine==false && @.Payload.Config_IsRunningCommandlet==false && @.Payload.LoadingSeconds>5 && @.Payload.LoadingSeconds<3600)]",
+			"filter": "$[?(@.Payload.EventName == 'Core.Loading' && @.Payload.LoadingName=='TotalEditorStartup' && @.Payload.Config_IsBuildMachine==false && @.Payload.Config_IsRunningCommandlet==false && @.Payload.LoadingSeconds>5 && @.Payload.LoadingSeconds<3600)]",
 			"property": "$.Payload.LoadingSeconds",
 			"groupBy": "$.Payload.Build_BranchName",
 			"function": "Average",
@@ -25,7 +25,7 @@
 		},
 		{
 			"id": "core-loading-user-loadmaptime",
-			"filter": "$[?(@.Payload.EventName == 'Core.Loading' && @.Payload.SchemaVersion>=2 && @.Payload.LoadingName=='LoadMap' && @.Payload.Config_IsBuildMachine==false && @.Payload.Config_IsRunningCommandlet==false && @.Payload.LoadingSeconds>5 && @.Payload.LoadingSeconds<3600)]",
+			"filter": "$[?(@.Payload.EventName == 'Core.Loading' && @.Payload.LoadingName=='LoadMap' && @.Payload.Config_IsBuildMachine==false && @.Payload.Config_IsRunningCommandlet==false && @.Payload.LoadingSeconds>5 && @.Payload.LoadingSeconds<3600)]",
 			"property": "$.Payload.LoadingSeconds",
 			"groupBy": "$.Payload.Build_BranchName,$.Payload.MapName",
 			"function": "Average",
@@ -33,14 +33,14 @@
 		},
 		{
 			"id": "core-loading-user-loadmap-count",
-			"filter": "$[?(@.Payload.EventName == 'Core.Loading' && @.Payload.SchemaVersion>=2 && @.Payload.LoadingName=='LoadMap' && @.Payload.Config_IsBuildMachine==false && @.Payload.Config_IsRunningCommandlet==false && @.Payload.LoadingSeconds>5 && @.Payload.LoadingSeconds<3600)]",
+			"filter": "$[?(@.Payload.EventName == 'Core.Loading' && @.Payload.LoadingName=='LoadMap' && @.Payload.Config_IsBuildMachine==false && @.Payload.Config_IsRunningCommandlet==false && @.Payload.LoadingSeconds>5 && @.Payload.LoadingSeconds<3600)]",
 			"groupBy": "$.Payload.Build_BranchName,$.Payload.MapName",
 			"function": "Count",
 			"interval": "00:15:00"
 		},
 		{
 			"id": "core-loading-user-pieloadmaptime",
-			"filter": "$[?(@.Payload.EventName == 'Core.Loading' && @.Payload.SchemaVersion>=2 && @.Payload.LoadingName=='PIE.LoadMapTime' && @.Payload.Config_IsBuildMachine==false && @.Payload.Config_IsRunningCommandlet==false && @.Payload.LoadingSeconds>5 && @.Payload.LoadingSeconds<3600)]",
+			"filter": "$[?(@.Payload.EventName == 'Core.Loading' && @.Payload.LoadingName=='PIE.LoadMapTime' && @.Payload.Config_IsBuildMachine==false && @.Payload.Config_IsRunningCommandlet==false && @.Payload.LoadingSeconds>5 && @.Payload.LoadingSeconds<3600)]",
 			"property": "$.Payload.LoadingSeconds",
 			"groupBy": "$.Payload.Build_BranchName,$.Payload.MapName",
 			"function": "Average",
@@ -48,7 +48,7 @@
 		},
 		{
 			"id": "core-loading-user-timetopie",
-			"filter": "$[?(@.Payload.EventName == 'Core.Loading' && @.Payload.SchemaVersion>=2 && @.Payload.LoadingName=='TimeToPIE' && @.Payload.Config_IsBuildMachine==false && @.Payload.Config_IsRunningCommandlet==false && @.Payload.LoadingSeconds>5 && @.Payload.LoadingSeconds<3600)]",
+			"filter": "$[?(@.Payload.EventName == 'Core.Loading' && @.Payload.LoadingName=='TimeToPIE' && @.Payload.Config_IsBuildMachine==false && @.Payload.Config_IsRunningCommandlet==false && @.Payload.LoadingSeconds>5 && @.Payload.LoadingSeconds<3600)]",
 			"property": "$.Payload.LoadingSeconds",
 			"groupBy": "$.Payload.Build_BranchName,$.Payload.MapName",
 			"function": "Average",
@@ -56,7 +56,7 @@
 		},
 		{
 			"id": "core-loading-user-piestartuptime",
-			"filter": "$[?(@.Payload.EventName == 'Core.Loading' && @.Payload.SchemaVersion>=2 && @.Payload.LoadingName=='PIE.TotalStartupTime' && @.Payload.Config_IsBuildMachine==false && @.Payload.Config_IsRunningCommandlet==false && @.Payload.LoadingSeconds>5 && @.Payload.LoadingSeconds<3600)]",
+			"filter": "$[?(@.Payload.EventName == 'Core.Loading' && @.Payload.LoadingName=='PIE.TotalStartupTime' && @.Payload.Config_IsBuildMachine==false && @.Payload.Config_IsRunningCommandlet==false && @.Payload.LoadingSeconds>5 && @.Payload.LoadingSeconds<3600)]",
 			"property": "$.Payload.LoadingSeconds",
 			"groupBy": "$.Payload.Build_BranchName,$.Payload.MapName",
 			"function": "Average",
@@ -64,7 +64,7 @@
 		},
 		{
 			"id": "core-loading-user-editorstartuptime-kpi",
-			"filter": "$[?(@.Payload.EventName == 'Core.Loading' && @.Payload.SchemaVersion>=3 && @.Payload.LoadingName=='TotalEditorStartup' && @.Payload.Config_IsBuildMachine==false && @.Payload.Config_IsRunningCommandlet==false && @.Payload.LoadingSeconds>5 && @.Payload.LoadingSeconds<3600)]",
+			"filter": "$[?(@.Payload.EventName == 'Core.Loading' && @.Payload.LoadingName=='TotalEditorStartup' && @.Payload.Config_IsBuildMachine==false && @.Payload.Config_IsRunningCommandlet==false && @.Payload.LoadingSeconds>5 && @.Payload.LoadingSeconds<3600)]",
 			"property": "$.Payload.LoadingSeconds",
 			"groupBy": "$.Payload.Build_BranchName",
 			"function": "Average",
@@ -72,7 +72,7 @@
 		},
 		{
 			"id": "core-loading-user-loadmaptime-kpi",
-			"filter": "$[?(@.Payload.EventName == 'Core.Loading' && @.Payload.SchemaVersion>=2 && @.Payload.LoadingName=='LoadMap' && @.Payload.Config_IsBuildMachine==false && @.Payload.Config_IsRunningCommandlet==false && @.Payload.LoadingSeconds>5 && @.Payload.LoadingSeconds<3600)]",
+			"filter": "$[?(@.Payload.EventName == 'Core.Loading' && @.Payload.LoadingName=='LoadMap' && @.Payload.Config_IsBuildMachine==false && @.Payload.Config_IsRunningCommandlet==false && @.Payload.LoadingSeconds>5 && @.Payload.LoadingSeconds<3600)]",
 			"property": "$.Payload.LoadingSeconds",
 			"groupBy": "$.Payload.Build_BranchName",
 			"function": "Average",
@@ -80,7 +80,7 @@
 		},
 		{
 			"id": "core-loading-user-piestartuptime-kpi",
-			"filter": "$[?(@.Payload.EventName == 'Core.Loading' && @.Payload.SchemaVersion>=2 && @.Payload.LoadingName=='PIE.TotalStartupTime' && @.Payload.Config_IsBuildMachine==false && @.Payload.Config_IsRunningCommandlet==false && @.Payload.LoadingSeconds>5 && @.Payload.LoadingSeconds<3600)]",
+			"filter": "$[?(@.Payload.EventName == 'Core.Loading' && @.Payload.LoadingName=='PIE.TotalStartupTime' && @.Payload.Config_IsBuildMachine==false && @.Payload.Config_IsRunningCommandlet==false && @.Payload.LoadingSeconds>5 && @.Payload.LoadingSeconds<3600)]",
 			"property": "$.Payload.LoadingSeconds",
 			"groupBy": "$.Payload.Build_BranchName",
 			"function": "Average",
@@ -88,7 +88,7 @@
 		},
 		{
 			"id": "core-loading-user-timetopie-kpi",
-			"filter": "$[?(@.Payload.EventName == 'Core.Loading' && @.Payload.SchemaVersion>=2 && @.Payload.LoadingName=='TimeToPIE' && @.Payload.Config_IsBuildMachine==false && @.Payload.Config_IsRunningCommandlet==false && @.Payload.LoadingSeconds>5 && @.Payload.LoadingSeconds<3600)]",
+			"filter": "$[?(@.Payload.EventName == 'Core.Loading' && @.Payload.LoadingName=='TimeToPIE' && @.Payload.Config_IsBuildMachine==false && @.Payload.Config_IsRunningCommandlet==false && @.Payload.LoadingSeconds>5 && @.Payload.LoadingSeconds<3600)]",
 			"property": "$.Payload.LoadingSeconds",
 			"groupBy": "$.Payload.Build_BranchName",
 			"function": "Average",
@@ -96,7 +96,7 @@
 		},
 		{
 			"id": "core-loading-user-cloudlocalddchitratio",
-			"filter": "$[?(@.Payload.EventName == 'Core.Loading' && @.Payload.SchemaVersion>=2 && @.Payload.Config_IsBuildMachine==false && @.Payload.DDC_Summary_CloudGetTotal>0 && @.Payload.LoadingSeconds>5 && @.Payload.LoadingSeconds<3600)]",
+			"filter": "$[?(@.Payload.EventName == 'Core.Loading' && @.Payload.Config_IsBuildMachine==false && @.Payload.DDC_Summary_CloudGetTotal>0 && @.Payload.LoadingSeconds>5 && @.Payload.LoadingSeconds<3600)]",
 			"property": "$.Payload.DDC_Summary_CloudGetHitPct",
 			"groupBy": "$.Payload.Build_BranchName",
 			"function": "Average",
@@ -104,7 +104,7 @@
 		},
 		{
 			"id": "core-loading-user-zenlocalddchitratio",
-			"filter": "$[?(@.Payload.EventName == 'Core.Loading' && @.Payload.SchemaVersion>=2 && @.Payload.Config_IsBuildMachine==false && @.Payload.Config_IsRunningCommandlet==false && @.Payload.DDC_Summary_ZenLocalGetTotal>0 && @.Payload.LoadingSeconds>5 && @.Payload.LoadingSeconds<3600)]",
+			"filter": "$[?(@.Payload.EventName == 'Core.Loading' && @.Payload.Config_IsBuildMachine==false && @.Payload.Config_IsRunningCommandlet==false && @.Payload.DDC_Summary_ZenLocalGetTotal>0 && @.Payload.LoadingSeconds>5 && @.Payload.LoadingSeconds<3600)]",
 			"property": "$.Payload.DDC_Summary_ZenLocalGetHitPct",
 			"groupBy": "$.Payload.Build_BranchName",
 			"function": "Average",
@@ -112,7 +112,7 @@
 		},
 		{
 			"id": "core-loading-user-sharedddchitratio",
-			"filter": "$[?(@.Payload.EventName == 'Core.Loading' && @.Payload.SchemaVersion>=2 && @.Payload.Config_IsBuildMachine==false && @.Payload.Config_IsRunningCommandlet==false && @.Payload.DDC_Summary_SharedGetTotal>0 && @.Payload.LoadingSeconds>5 && @.Payload.LoadingSeconds<3600)]",
+			"filter": "$[?(@.Payload.EventName == 'Core.Loading' && @.Payload.Config_IsBuildMachine==false && @.Payload.Config_IsRunningCommandlet==false && @.Payload.DDC_Summary_SharedGetTotal>0 && @.Payload.LoadingSeconds>5 && @.Payload.LoadingSeconds<3600)]",
 			"property": "$.Payload.DDC_Summary_SharedGetHitPct",
 			"groupBy": "$.Payload.Build_BranchName",
 			"function": "Average",
@@ -120,7 +120,7 @@
 		},
 		{
 			"id": "core-loading-user-totalddchitratio",
-			"filter": "$[?(@.Payload.EventName == 'Core.Loading' && @.Payload.SchemaVersion>=2 && @.Payload.Config_IsBuildMachine==false && @.Payload.Config_IsRunningCommandlet==false && @.Payload.LoadingSeconds>5 && @.Payload.LoadingSeconds<3600)]",
+			"filter": "$[?(@.Payload.EventName == 'Core.Loading' && @.Payload.Config_IsBuildMachine==false && @.Payload.Config_IsRunningCommandlet==false && @.Payload.LoadingSeconds>5 && @.Payload.LoadingSeconds<3600)]",
 			"property": "$.Payload.DDC_Summary_TotalGetHitPct",
 			"groupBy": "$.Payload.Build_BranchName",
 			"function": "Average",
@@ -128,7 +128,7 @@
 		},
 		{
 			"id": "core-loading-user-ddc-asset-efficiency",
-			"filter": "$[?(@.Payload.EventName=='Core.DDC.Resource' && @.Payload.SchemaVersion>=3 && @.Payload.Config_IsBuildMachine==false && @.Payload.Config_IsEditor==true && @.Payload.AssetType!='Total')]",
+			"filter": "$[?(@.Payload.EventName=='Core.DDC.Resource' && @.Payload.Config_IsBuildMachine==false && @.Payload.Config_IsEditor==true && @.Payload.AssetType!='Total')]",
 			"property": "$.Payload.Efficiency",
 			"groupBy": "$.Payload.Build_BranchName,$.Payload.AssetType",
 			"function": "Average",
@@ -136,7 +136,7 @@
 		},
 		{
 			"id": "core-loading-user-timetopie-percentile",
-			"filter": "$[?(@.Payload.EventName == 'Core.Loading' && @.Payload.SchemaVersion>=2 && @.Payload.LoadingName=='TimeToPIE' && @.Payload.Config_IsBuildMachine==false && @.Payload.Config_IsRunningCommandlet==false && @.Payload.LoadingSeconds>5 && @.Payload.LoadingSeconds<3600)]",
+			"filter": "$[?(@.Payload.EventName == 'Core.Loading' && @.Payload.LoadingName=='TimeToPIE' && @.Payload.Config_IsBuildMachine==false && @.Payload.Config_IsRunningCommandlet==false && @.Payload.LoadingSeconds>5 && @.Payload.LoadingSeconds<3600)]",
 			"groupBy": "$.Payload.Build_BranchName,$.Payload.MapName",
 			"property": "$.Payload.LoadingSeconds",
 			"function": "Percentile",
@@ -145,7 +145,7 @@
 		},
 		{
 			"id": "core-loading-user-editorstartuptime-percentile",
-			"filter": "$[?(@.Payload.EventName == 'Core.Loading' && @.Payload.SchemaVersion>=2 && @.Payload.LoadingName=='TotalEditorStartup' && @.Payload.Config_IsBuildMachine==false && @.Payload.Config_IsRunningCommandlet==false && @.Payload.LoadingSeconds>5 && @.Payload.LoadingSeconds<3600 )]",
+			"filter": "$[?(@.Payload.EventName == 'Core.Loading' && @.Payload.LoadingName=='TotalEditorStartup' && @.Payload.Config_IsBuildMachine==false && @.Payload.Config_IsRunningCommandlet==false && @.Payload.LoadingSeconds>5 && @.Payload.LoadingSeconds<3600 )]",
 			"groupBy": "$.Payload.Build_BranchName",
 			"property": "$.Payload.LoadingSeconds",
 			"function": "Percentile",
@@ -154,7 +154,7 @@
 		},
 		{
 			"id": "core-loading-user-loadmaptime-percentile",
-			"filter": "$[?(@.Payload.EventName == 'Core.Loading' && @.Payload.SchemaVersion>=3 && @.Payload.LoadingName=='LoadMap' && @.Payload.Config_IsBuildMachine==false && @.Payload.LoadingSeconds>5 && @.Payload.LoadingSeconds<3600)]",
+			"filter": "$[?(@.Payload.EventName == 'Core.Loading' && @.Payload.LoadingName=='LoadMap' && @.Payload.Config_IsBuildMachine==false && @.Payload.LoadingSeconds>5 && @.Payload.LoadingSeconds<3600)]",
 			"groupBy": "$.Payload.Build_BranchName,$.Payload.MapName",
 			"property": "$.Payload.LoadingSeconds",
 			"function": "Percentile",
@@ -163,7 +163,7 @@
 		},
 		{
 			"id": "core-loading-user-piestartuptime-percentile",
-			"filter": "$[?(@.Payload.EventName == 'Core.Loading' && @.Payload.SchemaVersion>=3 && @.Payload.LoadingName=='PIE.TotalStartupTime' && @.Payload.Config_IsBuildMachine==false && @.Payload.LoadingSeconds>5 && @.Payload.LoadingSeconds<3600)]",
+			"filter": "$[?(@.Payload.EventName == 'Core.Loading' && @.Payload.LoadingName=='PIE.TotalStartupTime' && @.Payload.Config_IsBuildMachine==false && @.Payload.LoadingSeconds>5 && @.Payload.LoadingSeconds<3600)]",
 			"groupBy": "$.Payload.Build_BranchName,$.Payload.MapName",
 			"property": "$.Payload.LoadingSeconds",
 			"function": "Percentile",
@@ -172,15 +172,23 @@
 		},
 		{
 			"id": "core-cooking-standard-cook-time",
-			"filter": "$[?(@.Payload.EventName == 'Core.Cooking' && @.Payload.SchemaVersion>=2 && search(@.Payload.Horde_TemplateName, \"Standard Build\") && search(@.Payload.Horde_StepName, \"Cook\"))]",
+			"filter": "$[?(@.Payload.EventName == 'Core.Cooking' && search(@.Payload.Horde_TemplateName, \"Standard Build\") && search(@.Payload.Horde_StepName, \"Cook\"))]",
 			"groupBy": "$.Payload.Build_BranchName,$.Payload.Horde_StepName,$.Payload.Cook_Profile_TargetPlatforms",
 			"property": "$.Payload.Cook_Profile_CookWallTimeSec",
 			"function": "Average",
 			"interval": "01:00:00"
 		},
 		{
+			"id": "core-cooking-standard-shaders-compiled",
+			"filter": "$[?(@.Payload.EventName == 'Core.Cooking' && search(@.Payload.Horde_TemplateName, \"Standard Build\"))]",
+			"groupBy": "$.Payload.Build_BranchName,$.Payload.Horde_StepName,$.Payload.Cook_Profile_TargetPlatforms",
+			"property": "$.Payload.Shaders_ShadersCompiled",
+			"function": "Average",
+			"interval": "01:00:00"
+		},
+		{
 			"id": "core-cooking-incremental-cook-time",
-			"filter": "$[?(@.Payload.EventName == 'Core.Cooking' && @.Payload.SchemaVersion>=3 && @.Payload.Horde_TemplateName=='Incremental CI')]",
+			"filter": "$[?(@.Payload.EventName == 'Core.Cooking' && @.Payload.Horde_TemplateName=='Incremental CI')]",
 			"groupBy": "$.Payload.Build_BranchName,$.Payload.Horde_StepName",
 			"property": "$.Payload.Cook_Profile_CookWallTimeSec",
 			"function": "Average",
@@ -188,23 +196,15 @@
 		},
 		{
 			"id": "core-cooking-incremental-shaders-compiled",
-			"filter": "$[?(@.Payload.EventName == 'Core.Cooking' && @.Payload.SchemaVersion>=3 && @.Payload.Horde_TemplateName=='Incremental CI')]",
+			"filter": "$[?(@.Payload.EventName == 'Core.Cooking' && @.Payload.Horde_TemplateName=='Incremental CI')]",
 			"groupBy": "$.Payload.Build_BranchName,$.Payload.Horde_StepName",
 			"property": "$.Payload.Shaders_ShadersCompiled",
 			"function": "Average",
 			"interval": "00:30:00"
 		},
 		{
-			"id": "core-cooking-nightly-cook-time",
-			"filter": "$[?(@.Payload.EventName == 'Core.Cooking' && @.Payload.SchemaVersion>=3 && search(@.Payload.Horde_TemplateName, \"Nightly Build\") && search(@.Payload.Horde_StepName, \"Cook\"))]",
-			"groupBy": "$.Payload.Build_BranchName,$.Payload.Horde_StepName,$.Payload.Cook_Profile_TargetPlatforms",
-			"property": "$.Payload.Cook_Profile_CookWallTimeSec",
-			"function": "Average",
-			"interval": "01:00:00"
-		},
-		{
 			"id": "core-cooking-nightly-cook-time-windows-kpi",
-			"filter": "$[?(@.Payload.EventName == 'Core.Cooking' && @.Payload.SchemaVersion>=3 && search(@.Payload.Horde_TemplateName, \"Nightly Build\") && @.Payload.Horde_StepName=='Cook Windows Client')]",
+			"filter": "$[?(@.Payload.EventName == 'Core.Cooking' && search(@.Payload.Horde_TemplateName, \"Nightly Build\") && @.Payload.Horde_StepName=='Cook Windows Client')]",
 			"groupBy": "$.Payload.Build_BranchName",
 			"property": "$.Payload.Cook_Profile_CookWallTimeSec",
 			"function": "Average",
@@ -212,31 +212,15 @@
 		},
 		{
 			"id": "core-cooking-nightly-cook-time-ps5-kpi",
-			"filter": "$[?(@.Payload.EventName == 'Core.Cooking' && @.Payload.SchemaVersion>=3 && search(@.Payload.Horde_TemplateName, \"Nightly Build\") && @.Payload.Horde_StepName=='Cook PS5 Client')]",
+			"filter": "$[?(@.Payload.EventName == 'Core.Cooking' && search(@.Payload.Horde_TemplateName, \"Nightly Build\") && @.Payload.Horde_StepName=='Cook PS5 Client')]",
 			"groupBy": "$.Payload.Build_BranchName",
 			"property": "$.Payload.Cook_Profile_CookWallTimeSec",
 			"function": "Average",
 			"interval": "01:00:00"
 		},
 		{
-			"id": "core-cooking-autotest-cook-time-kpi",
-			"filter": "$[?(@.Payload.EventName == 'Core.Cooking' && @.Payload.SchemaVersion>=3 && @.Payload.Horde_StepName=='AutoTest Cook')]",
-			"groupBy": "$.Payload.Build_BranchName",
-			"property": "$.Payload.Cook_Profile_CookWallTimeSec",
-			"function": "Average",
-			"interval": "01:00:00"
-		},
-		{
-			"id": "core-cooking-autotest-fastcook-time-kpi",
-			"filter": "$[?(@.Payload.EventName == 'Core.Cooking' && @.Payload.SchemaVersion>=3 && @.Payload.Horde_StepName=='AutoTest FastCook')]",
-			"groupBy": "$.Payload.Build_BranchName",
-			"property": "$.Payload.Cook_Profile_CookWallTimeSec",
-			"function": "Average",
-			"interval": "01:00:00"
-		},
-		{
 			"id": "core-cooking-nightly-cook-shaders-compiled-windows-kpi",
-			"filter": "$[?(@.Payload.EventName == 'Core.Cooking' && @.Payload.SchemaVersion>=3 && search(@.Payload.Horde_TemplateName, \"Nightly Build\") && @.Payload.Horde_StepName=='Cook Windows Client')]",
+			"filter": "$[?(@.Payload.EventName == 'Core.Cooking' && search(@.Payload.Horde_TemplateName, \"Nightly Build\") && @.Payload.Horde_StepName=='Cook Windows Client')]",
 			"groupBy": "$.Payload.Build_BranchName",
 			"property": "$.Payload.Shaders_ShadersCompiled",
 			"function": "Average",
@@ -244,47 +228,47 @@
 		},
 		{
 			"id": "core-cooking-nightly-cook-shaders-compiled-ps5-kpi",
-			"filter": "$[?(@.Payload.EventName == 'Core.Cooking' && @.Payload.SchemaVersion>=3 && search(@.Payload.Horde_TemplateName, \"Nightly Build\") && @.Payload.Horde_StepName=='Cook PS5 Client')]",
+			"filter": "$[?(@.Payload.EventName == 'Core.Cooking' && search(@.Payload.Horde_TemplateName, \"Nightly Build\") && @.Payload.Horde_StepName=='Cook PS5 Client')]",
 			"groupBy": "$.Payload.Build_BranchName",
 			"property": "$.Payload.Shaders_ShadersCompiled",
 			"function": "Average",
 			"interval": "01:00:00"
 		},
 		{
-			"id": "core-cooking-autotest-cook-shaders-compiled-kpi",
-			"filter": "$[?(@.Payload.EventName == 'Core.Cooking' && @.Payload.SchemaVersion>=3 && @.Payload.Horde_StepName=='AutoTest Cook')]",
-			"groupBy": "$.Payload.Build_BranchName",
-			"property": "$.Payload.Shaders_ShadersCompiled",
+			"id": "core-cooking-nightly-cook-time",
+			"filter": "$[?(@.Payload.EventName == 'Core.Cooking' && search(@.Payload.Horde_TemplateName, \"Nightly Build\") && search(@.Payload.Horde_StepName, \"Cook\"))]",
+			"groupBy": "$.Payload.Build_BranchName,$.Payload.Horde_StepName,$.Payload.Cook_Profile_TargetPlatforms",
+			"property": "$.Payload.Cook_Profile_CookWallTimeSec",
 			"function": "Average",
 			"interval": "01:00:00"
 		},
 		{
-			"id": "core-cooking-autotest-fastcook-shaders-compiled-kpi",
-			"filter": "$[?(@.Payload.EventName == 'Core.Cooking' && @.Payload.SchemaVersion>=3 && @.Payload.Horde_StepName=='AutoTest FastCook')]",
-			"groupBy": "$.Payload.Build_BranchName",
-			"property": "$.Payload.Shaders_ShadersCompiled",
-			"function": "Average",
-			"interval": "01:00:00"
-		},
-		{
-			"id": "core-cooking-standard-shaders-compiled",
-			"filter": "$[?(@.Payload.EventName == 'Core.Cooking' && @.Payload.SchemaVersion>=3 && search(@.Payload.Horde_TemplateName, \"Standard Build\"))]",
-			"groupBy": "$.Payload.Build_BranchName,$.Payload.Horde_StepName,$.Payload.Cook_Profile_TargetPlatforms",
-			"property": "$.Payload.Shaders_ShadersCompiled",
-			"function": "Average",
-			"interval": "01:00:00"
-		},
-		{
 			"id": "core-cooking-nightly-shaders-compiled",
-			"filter": "$[?(@.Payload.EventName == 'Core.Cooking' && @.Payload.SchemaVersion>=3 && search(@.Payload.Horde_TemplateName, \"Nightly Build\"))]",
+			"filter": "$[?(@.Payload.EventName == 'Core.Cooking' && search(@.Payload.Horde_TemplateName, \"Nightly Build\"))]",
 			"groupBy": "$.Payload.Build_BranchName,$.Payload.Horde_StepName,$.Payload.Cook_Profile_TargetPlatforms",
 			"property": "$.Payload.Shaders_ShadersCompiled",
 			"function": "Average",
 			"interval": "01:00:00"
 		},
 		{
+			"id": "core-cooking-autotest-cook-time-kpi",
+			"filter": "$[?(@.Payload.EventName == 'Core.Cooking' && @.Payload.Horde_StepName=='AutoTest Cook')]",
+			"groupBy": "$.Payload.Build_BranchName",
+			"property": "$.Payload.Cook_Profile_CookWallTimeSec",
+			"function": "Average",
+			"interval": "01:00:00"
+		},
+		{
+			"id": "core-cooking-autotest-cook-shaders-compiled-kpi",
+			"filter": "$[?(@.Payload.EventName == 'Core.Cooking' && @.Payload.Horde_StepName=='AutoTest Cook')]",
+			"groupBy": "$.Payload.Build_BranchName",
+			"property": "$.Payload.Shaders_ShadersCompiled",
+			"function": "Average",
+			"interval": "01:00:00"
+		},
+		{
 			"id": "core-cooking-autotest-cook-time",
-			"filter": "$[?(@.Payload.EventName == 'Core.Cooking' && @.Payload.SchemaVersion>=3 && @.Payload.Horde_StepName=='AutoTest Cook')]",
+			"filter": "$[?(@.Payload.EventName == 'Core.Cooking' && @.Payload.Horde_StepName=='AutoTest Cook')]",
 			"groupBy": "$.Payload.Build_BranchName,$.Payload.Horde_StepName,$.Payload.Cook_Profile_TargetPlatforms",
 			"property": "$.Payload.Cook_Profile_CookWallTimeSec",
 			"function": "Average",
@@ -292,17 +276,97 @@
 		},
 		{
 			"id": "core-cooking-autotest-shaders-compiled",
-			"filter": "$[?(@.Payload.EventName == 'Core.Cooking' && @.Payload.SchemaVersion>=3 && @.Payload.Horde_StepName=='AutoTest Cook')]",
+			"filter": "$[?(@.Payload.EventName == 'Core.Cooking' && @.Payload.Horde_StepName=='AutoTest Cook')]",
 			"groupBy": "$.Payload.Build_BranchName,$.Payload.Horde_StepName,$.Payload.Cook_Profile_TargetPlatforms",
 			"property": "$.Payload.Shaders_ShadersCompiled",
 			"function": "Average",
 			"interval": "00:30:00"
 		},
 		{
+			"id": "core-cooking-autotest-fastcook-time-kpi",
+			"filter": "$[?(@.Payload.EventName == 'Core.Cooking' && @.Payload.Horde_StepName=='AutoTest FastCook')]",
+			"groupBy": "$.Payload.Build_BranchName",
+			"property": "$.Payload.Cook_Profile_CookWallTimeSec",
+			"function": "Average",
+			"interval": "01:00:00"
+		},
+		{
+			"id": "core-cooking-autotest-fastcook-shaders-compiled-kpi",
+			"filter": "$[?(@.Payload.EventName == 'Core.Cooking' && @.Payload.Horde_StepName=='AutoTest FastCook')]",
+			"groupBy": "$.Payload.Build_BranchName",
+			"property": "$.Payload.Shaders_ShadersCompiled",
+			"function": "Average",
+			"interval": "01:00:00"
+		},
+		{
 			"id": "core-cooking-autotest-fastcook-time",
-			"filter": "$[?(@.Payload.EventName == 'Core.Cooking' && @.Payload.SchemaVersion>=3 && @.Payload.Horde_StepName=='AutoTest FastCook')]",
+			"filter": "$[?(@.Payload.EventName == 'Core.Cooking' && @.Payload.Horde_StepName=='AutoTest FastCook')]",
 			"groupBy": "$.Payload.Build_BranchName,$.Payload.Horde_StepName,$.Payload.Cook_Profile_TargetPlatforms",
 			"property": "$.Payload.Cook_Profile_CookWallTimeSec",
 			"function": "Average",
 			"interval": "01:00:00"
 		},
+		{
+			"id": "core-cooking-lyra-time-windows-kpi",
+			"filter": "$[?(@.Payload.EventName == 'Core.Cooking' && search(@.Payload.Horde_TemplateName, \"Package and Test\") && @.Payload.Horde_StepName=='Lyra Cook Win64')]",
+			"groupBy": "$.Payload.Build_BranchName",
+			"property": "$.Payload.Cook_Profile_CookWallTimeSec",
+			"function": "Average",
+			"interval": "01:00:00"
+		},
+		{
+			"id": "core-cooking-ancientgame-time-windows-kpi",
+			"filter": "$[?(@.Payload.EventName == 'Core.Cooking' && search(@.Payload.Horde_TemplateName, \"Package and Test\") && @.Payload.Horde_StepName=='AncientGame Cook Win64')]",
+			"groupBy": "$.Payload.Build_BranchName",
+			"property": "$.Payload.Cook_Profile_CookWallTimeSec",
+			"function": "Average",
+			"interval": "01:00:00"
+		},
+		{
+			"id": "core-cooking-citysample-time-windows-kpi",
+			"filter": "$[?(@.Payload.EventName == 'Core.Cooking' && search(@.Payload.Horde_TemplateName, \"Package and Test\") && @.Payload.Horde_StepName=='CitySample Cook Win64')]",
+			"groupBy": "$.Payload.Build_BranchName",
+			"property": "$.Payload.Cook_Profile_CookWallTimeSec",
+			"function": "Average",
+			"interval": "01:00:00"
+		},
+		{
+			"id": "core-cooking-lyra-time-ps5-kpi",
+			"filter": "$[?(@.Payload.EventName == 'Core.Cooking' && search(@.Payload.Horde_TemplateName, \"Package and Test\") && @.Payload.Horde_StepName=='Lyra Cook PS5')]",
+			"groupBy": "$.Payload.Build_BranchName",
+			"property": "$.Payload.Cook_Profile_CookWallTimeSec",
+			"function": "Average",
+			"interval": "01:00:00"
+		},
+		{
+			"id": "core-cooking-ancientgame-time-ps5-kpi",
+			"filter": "$[?(@.Payload.EventName == 'Core.Cooking' && search(@.Payload.Horde_TemplateName, \"Package and Test\") && @.Payload.Horde_StepName=='AncientGame Cook PS5')]",
+			"groupBy": "$.Payload.Build_BranchName",
+			"property": "$.Payload.Cook_Profile_CookWallTimeSec",
+			"function": "Average",
+			"interval": "01:00:00"
+		},
+		{
+			"id": "core-cooking-citysample-time-ps5-kpi",
+			"filter": "$[?(@.Payload.EventName == 'Core.Cooking' && search(@.Payload.Horde_TemplateName, \"Package and Test\") && @.Payload.Horde_StepName=='CitySample Cook PS5')]",
+			"groupBy": "$.Payload.Build_BranchName",
+			"property": "$.Payload.Cook_Profile_CookWallTimeSec",
+			"function": "Average",
+			"interval": "01:00:00"
+		},
+		{
+			"id": "core-cooking-package-and-test-cook-time",
+			"filter": "$[?(@.Payload.EventName == 'Core.Cooking' && search(@.Payload.Horde_TemplateName, \"Package and Test\") && search(@.Payload.Horde_StepName, \"Cook\"))]",
+			"groupBy": "$.Payload.Build_BranchName,$.Payload.Horde_StepName,$.Payload.Cook_Profile_TargetPlatforms",
+			"property": "$.Payload.Cook_Profile_CookWallTimeSec",
+			"function": "Average",
+			"interval": "01:00:00"
+		},
+		{
+			"id": "core-cooking-package-and-test-shaders-compiled",
+			"filter": "$[?(@.Payload.EventName == 'Core.Cooking' && search(@.Payload.Horde_TemplateName, \"Package and Test\"))]",
+			"groupBy": "$.Payload.Build_BranchName,$.Payload.Horde_StepName,$.Payload.Cook_Profile_TargetPlatforms",
+			"property": "$.Payload.Shaders_ShadersCompiled",
+			"function": "Average",
+			"interval": "01:00:00"
+		},
