Fix for Horde Anlaytics metrics

diff --git a/Engine/Build/Horde/Telemetry/metrics.telemetry.json b/Engine/Build/Horde/Telemetry/metrics.telemetry.json
--- a/Engine/Build/Horde/Telemetry/metrics.telemetry.json
+++ b/Engine/Build/Horde/Telemetry/metrics.telemetry.json
@@ -404,7 +404,7 @@
 		},
 		{
 			"id": "horde-compile-jobs",
-			"filter": "$[?(@.Payload.EventName == 'State.JobStepRef' && && @.Payload.JobName=='Editor, Tools & Monolithics' && search(@.Payload.JobName, \"Compile\") && @.Payload.Duration>0 && @.Payload.State=='Completed')]",
+			"filter": "$[?(@.Payload.EventName == 'State.JobStepRef' && @.Payload.JobName=='Editor, Tools & Monolithics' && search(@.Payload.JobName, \"Compile\") && @.Payload.Duration>0 && @.Payload.State=='Completed')]",
 			"groupBy": "$.Payload.StreamId,$.Payload.StepName",
 			"property": "$.Payload.Duration",
 			"function": "Average",
