Update for Horde Analytics
Filtering for Incremental CI Jobs, Added AutoTest Jobs

diff --git a/Engine/Build/Horde/Dashboard/analytics.dashboard.json b/Engine/Build/Horde/Dashboard/analytics.dashboard.json
--- a/Engine/Build/Horde/Dashboard/analytics.dashboard.json
+++ b/Engine/Build/Horde/Dashboard/analytics.dashboard.json
@@ -445,13 +445,21 @@
 							"name": "Incremental Jobs",
 							"metrics": [
 								{
-									"id": "horde-incremental-build-jobs",
+									"id": "horde-incremental-build-jobs"
 								},
 								{
 									"id": "horde-incremental-ci-jobs"
 								}
 							]
 						},
+						{
+							"name": "AutoTest Jobs",
+							"metrics": [
+								{
+									"id": "horde-autotest-build-jobs"
+								}
+							]
+						},
 						{
 							"name": "Compilation Jobs",
 							"metrics": [
diff --git a/Engine/Build/Horde/Telemetry/metrics.telemetry.json b/Engine/Build/Horde/Telemetry/metrics.telemetry.json
--- a/Engine/Build/Horde/Telemetry/metrics.telemetry.json
+++ b/Engine/Build/Horde/Telemetry/metrics.telemetry.json
@@ -380,7 +380,15 @@
 		},
 		{
 			"id": "horde-incremental-ci-jobs",
-			"filter": "$[?(@.Payload.EventName == 'State.JobStepRef' && @.Payload.JobName=='Incremental CI' && @.Payload.Duration>0 && @.Payload.State=='Completed')]",
+			"filter": "$[?(@.Payload.EventName == 'State.JobStepRef' && @.Payload.JobName=='Incremental CI' && search(@.Payload.StepName, \"Incremental\") && @.Payload.Duration>0 && @.Payload.State=='Completed')]",
+			"groupBy": "$.Payload.StreamId,$.Payload.StepName",
+			"property": "$.Payload.Duration",
+			"function": "Average",
+			"interval": "01:00:00"
+		},
+		{
+			"id": "horde-autotest-ci-jobs",
+			"filter": "$[?(@.Payload.EventName == 'State.JobStepRef' && @.Payload.JobName=='Incremental CI' && search(@.Payload.StepName, \"AutoTest\") && @.Payload.Duration>0 && @.Payload.State=='Completed')]",
 			"groupBy": "$.Payload.StreamId,$.Payload.StepName",
 			"property": "$.Payload.Duration",
 			"function": "Average",
