Fixed issue with Compile Job and AutoTest  Job Dashboards

diff --git a/Engine/Build/Horde/Dashboard/analytics.dashboard.json b/Engine/Build/Horde/Dashboard/analytics.dashboard.json
--- a/Engine/Build/Horde/Dashboard/analytics.dashboard.json
+++ b/Engine/Build/Horde/Dashboard/analytics.dashboard.json
@@ -456,7 +456,7 @@
 							"name": "AutoTest Jobs",
 							"metrics": [
 								{
-									"id": "horde-autotest-build-jobs"
+									"id": "horde-autotest-ci-jobs"
 								}
 							]
 						},
@@ -464,7 +464,10 @@
 							"name": "Compilation Jobs",
 							"metrics": [
 								{
-									"id": "horde-compile-jobs"
+									"id": "horde-incremental-compile-jobs"
+								},
+								{
+									"id": "horde-monolithic-compile-jobs"
 								}
 							]
 						},
diff --git a/Engine/Build/Horde/Telemetry/metrics.telemetry.json b/Engine/Build/Horde/Telemetry/metrics.telemetry.json
--- a/Engine/Build/Horde/Telemetry/metrics.telemetry.json
+++ b/Engine/Build/Horde/Telemetry/metrics.telemetry.json
@@ -411,8 +411,16 @@
 			"interval": "01:00:00"
 		},
 		{
-			"id": "horde-compile-jobs",
-			"filter": "$[?(@.Payload.EventName == 'State.JobStepRef' && @.Payload.JobName=='Editor, Tools & Monolithics' && search(@.Payload.JobName, \"Compile\") && @.Payload.Duration>0 && @.Payload.State=='Completed')]",
+			"id": "horde-monolithic-compile-jobs",
+			"filter": "$[?(@.Payload.EventName == 'State.JobStepRef' && @.Payload.JobName=='Editor, Tools & Monolithics' && search(@.Payload.StepName, \"Compile\") && @.Payload.Duration>0 && @.Payload.State=='Completed')]",
+			"groupBy": "$.Payload.StreamId,$.Payload.StepName",
+			"property": "$.Payload.Duration",
+			"function": "Average",
+			"interval": "01:00:00"
+		},
+		{
+			"id": "horde-incremental-compile-jobs",
+			"filter": "$[?(@.Payload.EventName == 'State.JobStepRef' && @.Payload.JobName=='Incremental CI' && search(@.Payload.StepName, \"Incremental Compile\") && @.Payload.Duration>0 && @.Payload.State=='Completed')]",
 			"groupBy": "$.Payload.StreamId,$.Payload.StepName",
 			"property": "$.Payload.Duration",
 			"function": "Average",
