Fix for Incremental-CI metric

diff --git a/Engine/Build/Horde/Telemetry/metrics.telemetry.json b/Engine/Build/Horde/Telemetry/metrics.telemetry.json
--- a/Engine/Build/Horde/Telemetry/metrics.telemetry.json
+++ b/Engine/Build/Horde/Telemetry/metrics.telemetry.json
@@ -372,7 +372,7 @@
 		},
 		{
 			"id": "horde-incremental-build-jobs",
-			"filter": "$[?(@.Payload.EventName == 'State.JobStepRef' && @.Payload.JobName=='Incremental Build' && @.Payload.Duration>0 && @.Payload.State=='Completed')]",
+			"filter": "$[?(@.Payload.EventName == 'State.JobStepRef' && @.Payload.JobName=='Incremental Build' && search(@.Payload.StepName, \"Incremental\") && @.Payload.Duration>0 && @.Payload.State=='Completed')]",
 			"groupBy": "$.Payload.StreamId,$.Payload.StepName",
 			"property": "$.Payload.Duration",
 			"function": "Average",
