UI Improvements to Horde Analytics Dashboard
Added Compile Category.
Renamed Cooking Category to Cook
Renamed DDC Cactegory to Cache

diff --git a/Engine/Build/Horde/Dashboard/analytics.dashboard.json b/Engine/Build/Horde/Dashboard/analytics.dashboard.json
--- a/Engine/Build/Horde/Dashboard/analytics.dashboard.json
+++ b/Engine/Build/Horde/Dashboard/analytics.dashboard.json
@@ -204,7 +204,7 @@
 					]
                 },
                 {
-                    "name": "Cooking",
+                    "name": "Cook",
 					"charts": [
 						{
 							"name": "Cook Time (KPIs)",
@@ -403,7 +403,7 @@
 					]
                 },
 				{
-					"name": "DDC",
+					"name": "Cache",
 					"charts": [
 						{
 							"name": "User DDC Backend Efficiency",
@@ -442,7 +442,7 @@
 					"name": "Build",
 					"charts": [
 						{
-							"name": "Incremental Jobs",
+							"name": "Incremental",
 							"metrics": [
 								{
 									"id": "horde-incremental-build-jobs"
@@ -453,7 +453,7 @@
 							]
 						},
 						{
-							"name": "AutoTest Jobs",
+							"name": "AutoTest",
 							"metrics": [
 								{
 									"id": "horde-autotest-ci-jobs"
@@ -461,29 +461,39 @@
 							]
 						},
 						{
-							"name": "Compilation Jobs",
+							"name": "Package And Test",
 							"metrics": [
 								{
-									"id": "horde-incremental-compile-jobs"
-								},
-								{
-									"id": "horde-monolithic-compile-jobs"
+									"id": "horde-package-and-test-jobs"
 								}
 							]
 						},
 						{
-							"name": "Package And Test Jobs",
+							"name": "Nighty Build",
 							"metrics": [
 								{
-									"id": "horde-package-and-test-jobs"
+									"id": "horde-nightly-build-jobs"
+								}
+							]
+						}
+					]
+				},
+				{
+					"name": "Compile",
+					"charts": [
+						{
+							"name": "Full",
+							"metrics": [
+								{
+									"id": "horde-incremental-compile-jobs"
 								}
 							]
 						},
 						{
-							"name": "Nighty Build Jobs",
+							"name": "Incremental",
 							"metrics": [
 								{
-									"id": "horde-nightly-build-jobs"
+									"id": "horde-incremental-compile-jobs"
 								}
 							]
 						}
