Fix for monolithic compile jobs in Horde Analytics

diff --git a/Engine/Build/Horde/Dashboard/analytics.dashboard.json b/Engine/Build/Horde/Dashboard/analytics.dashboard.json
--- a/Engine/Build/Horde/Dashboard/analytics.dashboard.json
+++ b/Engine/Build/Horde/Dashboard/analytics.dashboard.json
@@ -485,7 +485,7 @@
 							"name": "Full",
 							"metrics": [
 								{
-									"id": "horde-incremental-compile-jobs"
+									"id": "horde-monolithic-compile-jobs"
 								}
 							]
 						},
