Made names of Horde Analytics Charts unique

diff --git a/Engine/Build/Horde/Dashboard/analytics.dashboard.json b/Engine/Build/Horde/Dashboard/analytics.dashboard.json
--- a/Engine/Build/Horde/Dashboard/analytics.dashboard.json
+++ b/Engine/Build/Horde/Dashboard/analytics.dashboard.json
@@ -442,7 +442,7 @@
 					"name": "Build",
 					"charts": [
 						{
-							"name": "Incremental",
+							"name": "Incremental Build",
 							"metrics": [
 								{
 									"id": "horde-incremental-build-jobs"
@@ -453,7 +453,7 @@
 							]
 						},
 						{
-							"name": "AutoTest",
+							"name": "AutoTest Build",
 							"metrics": [
 								{
 									"id": "horde-autotest-ci-jobs"
@@ -461,7 +461,7 @@
 							]
 						},
 						{
-							"name": "Package And Test",
+							"name": "Package And Test Build",
 							"metrics": [
 								{
 									"id": "horde-package-and-test-jobs"
@@ -482,7 +482,7 @@
 					"name": "Compile",
 					"charts": [
 						{
-							"name": "Full",
+							"name": "Full Comile",
 							"metrics": [
 								{
 									"id": "horde-monolithic-compile-jobs"
@@ -490,7 +490,7 @@
 							]
 						},
 						{
-							"name": "Incremental",
+							"name": "Incremental Compile",
 							"metrics": [
 								{
 									"id": "horde-incremental-compile-jobs"
