Fixed typo in Horde Analytics Compile-Full Compile chart

diff --git a/Engine/Build/Horde/Dashboard/analytics.dashboard.json b/Engine/Build/Horde/Dashboard/analytics.dashboard.json
--- a/Engine/Build/Horde/Dashboard/analytics.dashboard.json
+++ b/Engine/Build/Horde/Dashboard/analytics.dashboard.json
@@ -482,7 +482,7 @@
 					"name": "Compile",
 					"charts": [
 						{
-							"name": "Full Comile",
+							"name": "Full Compile",
 							"metrics": [
 								{
 									"id": "horde-monolithic-compile-jobs"
