Added Project Name Filter to Horde Analytics Dashboards and Telemetry

diff --git a/Engine/Build/Horde/Dashboard/analytics.dashboard.json b/Engine/Build/Horde/Dashboard/analytics.dashboard.json
--- a/Engine/Build/Horde/Dashboard/analytics.dashboard.json
+++ b/Engine/Build/Horde/Dashboard/analytics.dashboard.json
@@ -13,12 +13,12 @@
 					]
 				},
 				{
-					"name": "LoadingName",
-					"group": "LoadingName",
+					"name": "ProjectName",
+					"group": "ProjectName",
 					"defaults": [
-						"TotalEditorStartup",
-						"LoadMap",
-						"TimeToPIE"
+						"AncientWorld",
+						"CitySample",
+						"Lyra"
 					]
 				},
 				{
@@ -30,6 +30,15 @@
 						"L_DefaultEditorOverview"
 					]
 				},
+				{
+					"name": "LoadingName",
+					"group": "LoadingName",
+					"defaults": [
+						"TotalEditorStartup",
+						"LoadMap",
+						"TimeToPIE"
+					]
+				},
 				{
 					"name": "AssetType",
 					"group": "AssetType"
diff --git a/Engine/Build/Horde/Telemetry/metrics.telemetry.json b/Engine/Build/Horde/Telemetry/metrics.telemetry.json
--- a/Engine/Build/Horde/Telemetry/metrics.telemetry.json
+++ b/Engine/Build/Horde/Telemetry/metrics.telemetry.json
@@ -4,14 +4,14 @@
 			"id": "core-loading-user-loadingtime",
 			"filter": "$[?(@.Payload.EventName == 'Core.Loading' && @.Payload.Config_IsBuildMachine==false && @.Payload.Config_IsRunningCommandlet==false && @.Payload.LoadingSeconds>5 && @.Payload.LoadingSeconds<3600)]",
 			"property": "$.Payload.LoadingSeconds",
-			"groupBy": "$.Payload.Build_BranchName,$.Payload.LoadingName",
+			"groupBy": "$.Payload.Build_BranchName,$.Payload.ProjectName,$.Payload.LoadingName",
 			"function": "Average",
 			"interval": "00:15:00"
 		},
 		{
 			"id": "core-loading-user-loadingcount",
 			"filter": "$[?(@.Payload.EventName == 'Core.Loading' &&@.Payload.Config_IsBuildMachine==false )]",
-			"groupBy": "$.Payload.Build_BranchName,$.Payload.LoadingName",
+			"groupBy": "$.Payload.Build_BranchName,$.Payload.ProjectName,$.Payload.LoadingName",
 			"function": "Count",
 			"interval": "00:15:00"
 		},
@@ -19,7 +19,7 @@
 			"id": "core-loading-user-editorstartuptime",
 			"filter": "$[?(@.Payload.EventName == 'Core.Loading' && @.Payload.LoadingName=='TotalEditorStartup' && @.Payload.Config_IsBuildMachine==false && @.Payload.Config_IsRunningCommandlet==false && @.Payload.LoadingSeconds>5 && @.Payload.LoadingSeconds<3600)]",
 			"property": "$.Payload.LoadingSeconds",
-			"groupBy": "$.Payload.Build_BranchName",
+			"groupBy": "$.Payload.Build_BranchName,$.Payload.ProjectName,$.Payload.ProjectName",
 			"function": "Average",
 			"interval": "00:15:00"
 		},
@@ -27,14 +27,14 @@
 			"id": "core-loading-user-loadmaptime",
 			"filter": "$[?(@.Payload.EventName == 'Core.Loading' && @.Payload.LoadingName=='LoadMap' && @.Payload.Config_IsBuildMachine==false && @.Payload.Config_IsRunningCommandlet==false && @.Payload.LoadingSeconds>5 && @.Payload.LoadingSeconds<3600)]",
 			"property": "$.Payload.LoadingSeconds",
-			"groupBy": "$.Payload.Build_BranchName,$.Payload.MapName",
+			"groupBy": "$.Payload.Build_BranchName,$.Payload.ProjectName,$.Payload.MapName",
 			"function": "Average",
 			"interval": "00:15:00"
 		},
 		{
 			"id": "core-loading-user-loadmap-count",
 			"filter": "$[?(@.Payload.EventName == 'Core.Loading' && @.Payload.LoadingName=='LoadMap' && @.Payload.Config_IsBuildMachine==false && @.Payload.Config_IsRunningCommandlet==false && @.Payload.LoadingSeconds>5 && @.Payload.LoadingSeconds<3600)]",
-			"groupBy": "$.Payload.Build_BranchName,$.Payload.MapName",
+			"groupBy": "$.Payload.Build_BranchName,$.Payload.ProjectName,$.Payload.MapName",
 			"function": "Count",
 			"interval": "00:15:00"
 		},
@@ -42,7 +42,7 @@
 			"id": "core-loading-user-pieloadmaptime",
 			"filter": "$[?(@.Payload.EventName == 'Core.Loading' && @.Payload.LoadingName=='PIE.LoadMapTime' && @.Payload.Config_IsBuildMachine==false && @.Payload.Config_IsRunningCommandlet==false && @.Payload.LoadingSeconds>5 && @.Payload.LoadingSeconds<3600)]",
 			"property": "$.Payload.LoadingSeconds",
-			"groupBy": "$.Payload.Build_BranchName,$.Payload.MapName",
+			"groupBy": "$.Payload.Build_BranchName,$.Payload.ProjectName,$.Payload.MapName",
 			"function": "Average",
 			"interval": "00:15:00"
 		},
@@ -50,7 +50,7 @@
 			"id": "core-loading-user-timetopie",
 			"filter": "$[?(@.Payload.EventName == 'Core.Loading' && @.Payload.LoadingName=='TimeToPIE' && @.Payload.Config_IsBuildMachine==false && @.Payload.Config_IsRunningCommandlet==false && @.Payload.LoadingSeconds>5 && @.Payload.LoadingSeconds<3600)]",
 			"property": "$.Payload.LoadingSeconds",
-			"groupBy": "$.Payload.Build_BranchName,$.Payload.MapName",
+			"groupBy": "$.Payload.Build_BranchName,$.Payload.ProjectName,$.Payload.MapName",
 			"function": "Average",
 			"interval": "00:15:00"
 		},
@@ -58,7 +58,7 @@
 			"id": "core-loading-user-piestartuptime",
 			"filter": "$[?(@.Payload.EventName == 'Core.Loading' && @.Payload.LoadingName=='PIE.TotalStartupTime' && @.Payload.Config_IsBuildMachine==false && @.Payload.Config_IsRunningCommandlet==false && @.Payload.LoadingSeconds>5 && @.Payload.LoadingSeconds<3600)]",
 			"property": "$.Payload.LoadingSeconds",
-			"groupBy": "$.Payload.Build_BranchName,$.Payload.MapName",
+			"groupBy": "$.Payload.Build_BranchName,$.Payload.ProjectName,$.Payload.MapName",
 			"function": "Average",
 			"interval": "00:15:00"
 		},
@@ -66,7 +66,7 @@
 			"id": "core-loading-user-editorstartuptime-kpi",
 			"filter": "$[?(@.Payload.EventName == 'Core.Loading' && @.Payload.LoadingName=='TotalEditorStartup' && @.Payload.Config_IsBuildMachine==false && @.Payload.Config_IsRunningCommandlet==false && @.Payload.LoadingSeconds>5 && @.Payload.LoadingSeconds<3600)]",
 			"property": "$.Payload.LoadingSeconds",
-			"groupBy": "$.Payload.Build_BranchName",
+			"groupBy": "$.Payload.Build_BranchName,$.Payload.ProjectName",
 			"function": "Average",
 			"interval": "01:00:00"
 		},
@@ -74,7 +74,7 @@
 			"id": "core-loading-user-loadmaptime-kpi",
 			"filter": "$[?(@.Payload.EventName == 'Core.Loading' && @.Payload.LoadingName=='LoadMap' && @.Payload.Config_IsBuildMachine==false && @.Payload.Config_IsRunningCommandlet==false && @.Payload.LoadingSeconds>5 && @.Payload.LoadingSeconds<3600)]",
 			"property": "$.Payload.LoadingSeconds",
-			"groupBy": "$.Payload.Build_BranchName",
+			"groupBy": "$.Payload.Build_BranchName,$.Payload.ProjectName",
 			"function": "Average",
 			"interval": "01:00:00"
 		},
@@ -82,7 +82,7 @@
 			"id": "core-loading-user-piestartuptime-kpi",
 			"filter": "$[?(@.Payload.EventName == 'Core.Loading' && @.Payload.LoadingName=='PIE.TotalStartupTime' && @.Payload.Config_IsBuildMachine==false && @.Payload.Config_IsRunningCommandlet==false && @.Payload.LoadingSeconds>5 && @.Payload.LoadingSeconds<3600)]",
 			"property": "$.Payload.LoadingSeconds",
-			"groupBy": "$.Payload.Build_BranchName",
+			"groupBy": "$.Payload.Build_BranchName,$.Payload.ProjectName",
 			"function": "Average",
 			"interval": "01:00:00"
 		},
@@ -90,7 +90,7 @@
 			"id": "core-loading-user-timetopie-kpi",
 			"filter": "$[?(@.Payload.EventName == 'Core.Loading' && @.Payload.LoadingName=='TimeToPIE' && @.Payload.Config_IsBuildMachine==false && @.Payload.Config_IsRunningCommandlet==false && @.Payload.LoadingSeconds>5 && @.Payload.LoadingSeconds<3600)]",
 			"property": "$.Payload.LoadingSeconds",
-			"groupBy": "$.Payload.Build_BranchName",
+			"groupBy": "$.Payload.Build_BranchName,$.Payload.ProjectName",
 			"function": "Average",
 			"interval": "01:00:00"
 		},
@@ -98,7 +98,7 @@
 			"id": "core-loading-user-cloudlocalddchitratio",
 			"filter": "$[?(@.Payload.EventName == 'Core.Loading' && @.Payload.Config_IsBuildMachine==false && @.Payload.DDC_Summary_CloudGetTotal>0 && @.Payload.LoadingSeconds>5 && @.Payload.LoadingSeconds<3600)]",
 			"property": "$.Payload.DDC_Summary_CloudGetHitPct",
-			"groupBy": "$.Payload.Build_BranchName",
+			"groupBy": "$.Payload.Build_BranchName,$.Payload.ProjectName",
 			"function": "Average",
 			"interval": "01:00:00"
 		},
@@ -106,7 +106,7 @@
 			"id": "core-loading-user-zenlocalddchitratio",
 			"filter": "$[?(@.Payload.EventName == 'Core.Loading' && @.Payload.Config_IsBuildMachine==false && @.Payload.Config_IsRunningCommandlet==false && @.Payload.DDC_Summary_ZenLocalGetTotal>0 && @.Payload.LoadingSeconds>5 && @.Payload.LoadingSeconds<3600)]",
 			"property": "$.Payload.DDC_Summary_ZenLocalGetHitPct",
-			"groupBy": "$.Payload.Build_BranchName",
+			"groupBy": "$.Payload.Build_BranchName,$.Payload.ProjectName",
 			"function": "Average",
 			"interval": "01:00:00"
 		},
@@ -114,7 +114,7 @@
 			"id": "core-loading-user-sharedddchitratio",
 			"filter": "$[?(@.Payload.EventName == 'Core.Loading' && @.Payload.Config_IsBuildMachine==false && @.Payload.Config_IsRunningCommandlet==false && @.Payload.DDC_Summary_SharedGetTotal>0 && @.Payload.LoadingSeconds>5 && @.Payload.LoadingSeconds<3600)]",
 			"property": "$.Payload.DDC_Summary_SharedGetHitPct",
-			"groupBy": "$.Payload.Build_BranchName",
+			"groupBy": "$.Payload.Build_BranchName,$.Payload.ProjectName",
 			"function": "Average",
 			"interval": "01:00:00"
 		},
@@ -122,7 +122,7 @@
 			"id": "core-loading-user-totalddchitratio",
 			"filter": "$[?(@.Payload.EventName == 'Core.Loading' && @.Payload.Config_IsBuildMachine==false && @.Payload.Config_IsRunningCommandlet==false && @.Payload.LoadingSeconds>5 && @.Payload.LoadingSeconds<3600)]",
 			"property": "$.Payload.DDC_Summary_TotalGetHitPct",
-			"groupBy": "$.Payload.Build_BranchName",
+			"groupBy": "$.Payload.Build_BranchName,$.Payload.ProjectName",
 			"function": "Average",
 			"interval": "01:00:00"
 		},
@@ -130,14 +130,14 @@
 			"id": "core-loading-user-ddc-asset-efficiency",
 			"filter": "$[?(@.Payload.EventName=='Core.DDC.Resource' && @.Payload.Config_IsBuildMachine==false && @.Payload.Config_IsEditor==true && @.Payload.AssetType!='Total')]",
 			"property": "$.Payload.Efficiency",
-			"groupBy": "$.Payload.Build_BranchName,$.Payload.AssetType",
+			"groupBy": "$.Payload.Build_BranchName,$.Payload.ProjectName,$.Payload.AssetType",
 			"function": "Average",
 			"interval": "01:00:00"
 		},
 		{
 			"id": "core-loading-user-timetopie-percentile",
 			"filter": "$[?(@.Payload.EventName == 'Core.Loading' && @.Payload.LoadingName=='TimeToPIE' && @.Payload.Config_IsBuildMachine==false && @.Payload.Config_IsRunningCommandlet==false && @.Payload.LoadingSeconds>5 && @.Payload.LoadingSeconds<3600)]",
-			"groupBy": "$.Payload.Build_BranchName,$.Payload.MapName",
+			"groupBy": "$.Payload.Build_BranchName,$.Payload.ProjectName,$.Payload.MapName",
 			"property": "$.Payload.LoadingSeconds",
 			"function": "Percentile",
 			"percentile": 75,
@@ -146,7 +146,7 @@
 		{
 			"id": "core-loading-user-editorstartuptime-percentile",
 			"filter": "$[?(@.Payload.EventName == 'Core.Loading' && @.Payload.LoadingName=='TotalEditorStartup' && @.Payload.Config_IsBuildMachine==false && @.Payload.Config_IsRunningCommandlet==false && @.Payload.LoadingSeconds>5 && @.Payload.LoadingSeconds<3600 )]",
-			"groupBy": "$.Payload.Build_BranchName",
+			"groupBy": "$.Payload.Build_BranchName,$.Payload.ProjectName",
 			"property": "$.Payload.LoadingSeconds",
 			"function": "Percentile",
 			"percentile": 75,
@@ -155,7 +155,7 @@
 		{
 			"id": "core-loading-user-loadmaptime-percentile",
 			"filter": "$[?(@.Payload.EventName == 'Core.Loading' && @.Payload.LoadingName=='LoadMap' && @.Payload.Config_IsBuildMachine==false && @.Payload.LoadingSeconds>5 && @.Payload.LoadingSeconds<3600)]",
-			"groupBy": "$.Payload.Build_BranchName,$.Payload.MapName",
+			"groupBy": "$.Payload.Build_BranchName,$.Payload.ProjectName,$.Payload.MapName",
 			"property": "$.Payload.LoadingSeconds",
 			"function": "Percentile",
 			"percentile": 75,
@@ -164,7 +164,7 @@
 		{
 			"id": "core-loading-user-piestartuptime-percentile",
 			"filter": "$[?(@.Payload.EventName == 'Core.Loading' && @.Payload.LoadingName=='PIE.TotalStartupTime' && @.Payload.Config_IsBuildMachine==false && @.Payload.LoadingSeconds>5 && @.Payload.LoadingSeconds<3600)]",
-			"groupBy": "$.Payload.Build_BranchName,$.Payload.MapName",
+			"groupBy": "$.Payload.Build_BranchName,$.Payload.ProjectName,$.Payload.MapName",
 			"property": "$.Payload.LoadingSeconds",
 			"function": "Percentile",
 			"percentile": 75,
@@ -197,7 +197,7 @@
 		{
 			"id": "core-cooking-incremental-shaders-compiled",
 			"filter": "$[?(@.Payload.EventName == 'Core.Cooking' && @.Payload.Horde_TemplateName=='Incremental CI')]",
-			"groupBy": "$.Payload.Build_BranchName,$.Payload.Horde_StepName",
+			"groupBy": "$.Payload.Build_BranchName,$.Payload.ProjectName,$.Payload.Horde_StepName",
 			"property": "$.Payload.Shaders_ShadersCompiled",
 			"function": "Average",
 			"interval": "00:30:00"
@@ -293,7 +293,7 @@
 		{
 			"id": "core-cooking-autotest-fastcook-shaders-compiled-kpi",
 			"filter": "$[?(@.Payload.EventName == 'Core.Cooking' && @.Payload.Horde_StepName=='AutoTest FastCook')]",
-			"groupBy": "$.Payload.Build_BranchName",
+			"groupBy": "$.Payload.Build_BranchName,$.Payload.ProjectName",
 			"property": "$.Payload.Shaders_ShadersCompiled",
 			"function": "Average",
 			"interval": "01:00:00"
