Changed some grouping names in Horde Analytics

diff --git a/Engine/Build/Horde/Dashboard/analytics.dashboard.json b/Engine/Build/Horde/Dashboard/analytics.dashboard.json
--- a/Engine/Build/Horde/Dashboard/analytics.dashboard.json
+++ b/Engine/Build/Horde/Dashboard/analytics.dashboard.json
@@ -13,7 +13,7 @@
 					]
 				},
 				{
-					"name": "ProjectName",
+					"name": "Project",
 					"group": "ProjectName",
 					"defaults": [
 						"AncientWorld",
@@ -22,7 +22,7 @@
 					]
 				},
 				{
-					"name": "MapName",
+					"name": "Map",
 					"group": "MapName",
 					"defaults": [
 						"AncientWorld",
@@ -44,7 +44,7 @@
 					"group": "AssetType"
 				},
 				{
-					"name": "TargetPlatform",
+					"name": "Cook Platform",
 					"group": "Cook_Profile_TargetPlatforms"
 				},
 				{
