Added Editor Performance Warning Chart to Horde Analytics

diff --git a/Engine/Build/Horde/Dashboard/analytics.dashboard.json b/Engine/Build/Horde/Dashboard/analytics.dashboard.json
--- a/Engine/Build/Horde/Dashboard/analytics.dashboard.json
+++ b/Engine/Build/Horde/Dashboard/analytics.dashboard.json
@@ -130,6 +130,14 @@
 							"graph": "indicator",
 							"max": 900
 						},
+						{
+							"name": "Performance Warnings",
+							"metrics": [
+								{
+									"id": "editor-performance-warning"
+								}
+							]
+						},
 						{
 							"name": "User Loading Time (KPIs)",
 							"metrics": [
@@ -154,7 +162,7 @@
 						{
 							"name": "User Loading Count (Live)",
 							"display": "Value",
-							"max" : 100,
+							"max": 100,
 							"metrics": [
 								{
 									"id": "core-loading-user-loadingcount"
diff --git a/Engine/Build/Horde/Telemetry/metrics.telemetry.json b/Engine/Build/Horde/Telemetry/metrics.telemetry.json
--- a/Engine/Build/Horde/Telemetry/metrics.telemetry.json
+++ b/Engine/Build/Horde/Telemetry/metrics.telemetry.json
@@ -1,5 +1,13 @@
 {
 	"metrics": [
+		{
+			"id": "editor-performance-warning",
+			"filter": "$[?(@.Payload.EventName == 'Editor.PerformanceWarning')]",
+			"groupBy": "$.Payload.Build_BranchName,$.Payload.ProjectName,$.Payload.KPI_Category,$.Payload.KPI_Name",
+			"property": "$.Payload.KPI_CurrentValue",
+			"function": "Average",
+			"interval": "00:15:00"
+		},
 		{
 			"id": "core-loading-user-loadingtime",
 			"filter": "$[?(@.Payload.EventName == 'Core.Loading' && @.Payload.Config_IsBuildMachine==false && @.Payload.Config_IsRunningCommandlet==false && @.Payload.LoadingSeconds>5 && @.Payload.LoadingSeconds<3600)]",
