Further tweaks to Horde Analytics Dashboard
Adedd Wanring Count to Editor Performance Category

diff --git a/Engine/Build/Horde/Dashboard/analytics.dashboard.json b/Engine/Build/Horde/Dashboard/analytics.dashboard.json
--- a/Engine/Build/Horde/Dashboard/analytics.dashboard.json
+++ b/Engine/Build/Horde/Dashboard/analytics.dashboard.json
@@ -72,9 +72,9 @@
 					"group": "StepName"
 				}
 			],
-            "categories": [
-                {
-                    "name": "Editor Iteration",
+			"categories": [
+				{
+					"name": "Editor Iteration",
 					"charts": [
 						{
 							"name": "User Load Time (75th Percentile)",
@@ -219,9 +219,50 @@
 							]
 						}
 					]
-                },
-                {
-                    "name": "Cook",
+				},
+				{
+					"name": "Editor Performance",
+					"charts": [
+						{
+							"name": "Count",
+							"display": "Value",
+							"metrics": [
+								{
+									"id": "editor-performance-warning-count"
+								}
+							]
+						},
+						{
+							"name": "Editor",
+							"display": "Value",
+							"metrics": [
+								{
+									"id": "editor-performance-editor-warning"
+								}
+							]
+						},
+						{
+							"name": "PIE",
+							"display": "Value",
+							"metrics": [
+								{
+									"id": "editor-performance-pie-warning"
+								}
+							]
+						},
+						{
+							"name": "Cache",
+							"display": "Value",
+							"metrics": [
+								{
+									"id": "editor-performance-cache-warning"
+								}
+							]
+						}
+					]
+				},
+				{
+					"name": "Cook",
 					"charts": [
 						{
 							"name": "Cook Time (KPIs)",
@@ -418,7 +459,7 @@
 							]
 						}
 					]
-                },
+				},
 				{
 					"name": "Cache",
 					"charts": [
@@ -516,7 +557,7 @@
 						}
 					]
 				}
-            ]
+			]
         }
     ]
 }
\ No newline at end of file
diff --git a/Engine/Build/Horde/Telemetry/metrics.telemetry.json b/Engine/Build/Horde/Telemetry/metrics.telemetry.json
--- a/Engine/Build/Horde/Telemetry/metrics.telemetry.json
+++ b/Engine/Build/Horde/Telemetry/metrics.telemetry.json
@@ -1,13 +1,36 @@
 {
 	"metrics": [
 		{
-			"id": "editor-performance-warning",
-			"filter": "$[?(@.Payload.EventName == 'Editor.PerformanceWarning')]",
-			"groupBy": "$.Payload.Build_BranchName,$.Payload.ProjectName,$.Payload.KPI_Category,$.Payload.KPI_Name",
+			"id": "editor-performance-cache-warning",
+			"filter": "$[?(@.Payload.EventName == 'Editor.PerformanceWarning' && @.Payload.KPI_Category=='Cache' && @.Payload.KPI_DisplayType=='Percent')]",
+			"groupBy": "$.Payload.Build_BranchName,$.Payload.ProjectName,$.Payload.KPI_Name",
+			"property": "$.Payload.KPI_CurrentValue",
+			"function": "Average",
+			"interval": "00:15:00"
+		},
+		{
+			"id": "editor-performance-editor-warning",
+			"filter": "$[?(@.Payload.EventName == 'Editor.PerformanceWarning' && @.Payload.KPI_Category=='Editor')]",
+			"groupBy": "$.Payload.Build_BranchName,$.Payload.ProjectName,$.Payload.MapName,$.Payload.KPI_Name",
+			"property": "$.Payload.KPI_CurrentValue",
+			"function": "Average",
+			"interval": "00:15:00"
+		},
+		{
+			"id": "editor-performance-pie-warning",
+			"filter": "$[?(@.Payload.EventName == 'Editor.PerformanceWarning' && @.Payload.KPI_Category=='PIE')]",
+			"groupBy": "$.Payload.Build_BranchName,$.Payload.ProjectName,$.Payload.MapName,$.Payload.KPI_Name",
 			"property": "$.Payload.KPI_CurrentValue",
 			"function": "Average",
 			"interval": "00:15:00"
 		},
+		{
+			"id": "editor-performance-warning-count",
+			"filter": "$[?(@.Payload.EventName == 'Editor.PerformanceWarning')]",
+			"groupBy": "$.Payload.Build_BranchName,$.Payload.ProjectName,$.Payload.KPI_Category,$.Payload.KPI_Name",
+			"function": "Count",
+			"interval": "00:15:00"
+		},
 		{
 			"id": "core-loading-user-loadingtime",
 			"filter": "$[?(@.Payload.EventName == 'Core.Loading' && @.Payload.Config_IsBuildMachine==false && @.Payload.Config_IsRunningCommandlet==false && @.Payload.LoadingSeconds>5 && @.Payload.LoadingSeconds<3600)]",
