Horde: Engine analytics updates
#rnx

diff --git a/Engine/Build/Horde/Dashboard/analytics.dashboard.json b/Engine/Build/Horde/Dashboard/analytics.dashboard.json
--- a/Engine/Build/Horde/Dashboard/analytics.dashboard.json
+++ b/Engine/Build/Horde/Dashboard/analytics.dashboard.json
@@ -1,8 +1,9 @@
 {
-    "analytics": [
-        {
-            "id": "core-loading-user",
-            "name": "Core Loading User",
+	"analytics": [
+		{
+			"id": "engine",
+			"name": "UE5",
+			"telemetryStoreId": "engine",
 			"variables": [
 				{
 					"name": "Stream",
@@ -558,6 +559,6 @@
 					]
 				}
 			]
-        }
-    ]
+		}
+	]
 }
\ No newline at end of file
