Add empty config files for analytics and metrics to fix server config update.

#rnx

diff --git a/Engine/Build/Horde/Dashboard/analytics.dashboard.json b/Engine/Build/Horde/Dashboard/analytics.dashboard.json
new file mode 100644
--- /dev/null
+++ b/Engine/Build/Horde/Dashboard/analytics.dashboard.json
@@ -0,0 +1,2 @@
+{
+}
diff --git a/Engine/Build/Horde/Telemetry/metrics.telemetry.json b/Engine/Build/Horde/Telemetry/metrics.telemetry.json
new file mode 100644
--- /dev/null
+++ b/Engine/Build/Horde/Telemetry/metrics.telemetry.json
@@ -0,0 +1,2 @@
+{
+}
